Fix inadvertent change that broke a redirector

GitHub pull request #28
Fix typo in K2FieldRedirects config: (Acor -> Actor)

diff --git a/Engine/Config/BaseEngine.ini b/Engine/Config/BaseEngine.ini
--- a/Engine/Config/BaseEngine.ini
+++ b/Engine/Config/BaseEngine.ini
@@ -624,7 +624,7 @@ bUseStreamingPause=false
 +K2FieldRedirects=(OldFieldName="SceneComponent.GetLocation",NewFieldName="K2_GetComponentLocation")
 +K2FieldRedirects=(OldFieldName="SceneComponent.GetRotation",NewFieldName="K2_GetComponentRotation")
 +K2FieldRedirects=(OldFieldName="SceneComponent.GetComponentToWorld",NewFieldName="K2_GetComponentToWorld")
-+K2FieldRedirects=(OldFieldName="Acor.Destroy",NewFieldName="K2_DestroyActor")
++K2FieldRedirects=(OldFieldName="Actor.Destroy",NewFieldName="K2_DestroyActor")
 +K2FieldRedirects=(OldFieldName="Actor.GetControllerAnalogStickState",NewFieldName="InputComponent.GetControllerAnalogStickState")
 +K2FieldRedirects=(OldFieldName="Actor.HasTag", NewFieldName="Actor.ActorHasTag")
 
