#ttp 331319: CRITICAL: LIVE: FRAMEWORK: GitHub editor attempts to connect to shared network DDC

#change Moved internal DDC shared paths to a no-redist module.

diff --git a/Engine/Config/BaseEngine.ini b/Engine/Config/BaseEngine.ini
--- a/Engine/Config/BaseEngine.ini
+++ b/Engine/Config/BaseEngine.ini
@@ -1746,8 +1746,8 @@ AsyncPut=(Type=AsyncPut, Inner=Hierarchy)
 Hierarchy=(Type=Hierarchical, Inner=Boot, Inner=Pak, Inner=EnginePak, Inner=Local, Inner=Shared)
 Boot=(Type=Boot, Filename="%GAMEDIR%DerivedDataCache/Boot.ddc", MaxCacheSize=512)
 Local=(Type=FileSystem, ReadOnly=false, Clean=false, Flush=false, PurgeTransient=true, DeleteUnused=true, UnusedFileAge=34, FoldersToClean=-1, Path=../../../Engine/DerivedDataCache, EnvPathOverride=UE-LocalDataCachePath)
-Shared=(Type=FileSystem, ReadOnly=false, Clean=false, Flush=false, DeleteUnused=true, UnusedFileAge=23, FoldersToClean=-1, Path=\\epicgames.net\root\UE4DDC, EnvPathOverride=UE-SharedDataCachePath)
-AltShared=(Type=FileSystem, ReadOnly=true, Clean=false, Flush=false, DeleteUnused=true, UnusedFileAge=23, FoldersToClean=-1, Path=\\epicgames.net\root\UE4DDC2, EnvPathOverride=UE-SharedDataCachePath2)
+Shared=(Type=FileSystem, ReadOnly=false, Clean=false, Flush=false, DeleteUnused=true, UnusedFileAge=23, FoldersToClean=-1, Path=?EpicDDC, EnvPathOverride=UE-SharedDataCachePath)
+AltShared=(Type=FileSystem, ReadOnly=true, Clean=false, Flush=false, DeleteUnused=true, UnusedFileAge=23, FoldersToClean=-1, Path=?EpicDDC2, EnvPathOverride=UE-SharedDataCachePath2)
 Pak=(Type=ReadPak, Filename="%GAMEDIR%DerivedDataCache/DDC.ddp")
 EnginePak=(Type=ReadPak, Filename=../../../Engine/DerivedDataCache/DDC.ddp)
 
diff --git a/Engine/Source/Developer/DerivedDataCache/DerivedDataCache.Build.cs b/Engine/Source/Developer/DerivedDataCache/DerivedDataCache.Build.cs
--- a/Engine/Source/Developer/DerivedDataCache/DerivedDataCache.Build.cs
+++ b/Engine/Source/Developer/DerivedDataCache/DerivedDataCache.Build.cs
@@ -1,11 +1,17 @@
 // Copyright 1998-2014 Epic Games, Inc. All Rights Reserved.
 
 using UnrealBuildTool;
+using System.IO;
 
 public class DerivedDataCache : ModuleRules
 {
 	public DerivedDataCache(TargetInfo Target)
 	{
 		PrivateDependencyModuleNames.Add("Core");
+		// Internal (NoRedist) module
+		if (Directory.Exists(Path.Combine("Developer", "NoRedist", "DDCUtils")))
+		{
+			DynamicallyLoadedModuleNames.Add("DDCUtils");
+		}
 	}
 }
diff --git a/Engine/Source/Developer/DerivedDataCache/Private/DerivedDataBackends.cpp b/Engine/Source/Developer/DerivedDataCache/Private/DerivedDataBackends.cpp
--- a/Engine/Source/Developer/DerivedDataCache/Private/DerivedDataBackends.cpp
+++ b/Engine/Source/Developer/DerivedDataCache/Private/DerivedDataBackends.cpp
@@ -11,6 +11,8 @@
 #include "DerivedDataLimitKeyLengthWrapper.h"
 #include "DerivedDataBackendCorruptionWrapper.h"
 #include "DerivedDataBackendVerifyWrapper.h"
+#include "DerivedDataUtilsInterface.h"
+#include "Misc/EngineBuildSettings.h"
 
 DEFINE_LOG_CATEGORY(LogDerivedDataCache);
 
@@ -475,6 +477,23 @@ public:
 				UE_LOG( LogDerivedDataCache, Log, TEXT("Found environment variable %s=%s"), *EnvPathOverride, *Path );
 			}
 		}
+		// Check if the Path is a real path or a special keyword
+		if (FEngineBuildSettings::IsInternalBuild())
+		{
+			auto DDCUtils = FModuleManager::LoadModulePtr< IDDCUtilsModuleInterface >("DDCUtils");
+			if (DDCUtils)
+			{
+				FString PathFromName = DDCUtils->GetSharedCachePath(Path);
+				if (!PathFromName.IsEmpty())
+				{
+					Path = PathFromName;
+				}
+			}
+		}
+		else if (Path.StartsWith(TEXT("?")))
+		{
+			Path = TEXT("");
+		}
 
 		if( !Path.Len() )
 		{
diff --git a/Engine/Source/Developer/DerivedDataCache/Public/DerivedDataUtilsInterface.h b/Engine/Source/Developer/DerivedDataCache/Public/DerivedDataUtilsInterface.h
new file mode 100644
--- /dev/null
+++ b/Engine/Source/Developer/DerivedDataCache/Public/DerivedDataUtilsInterface.h
@@ -0,0 +1,14 @@
+// Copyright 1998-2014 Epic Games, Inc. All Rights Reserved.
+#pragma once
+
+/** 
+ * Interface for DDC utility modules.
+ **/
+class IDDCUtilsModuleInterface : public IModuleInterface
+{
+public:
+	/**
+	 * Returns the path to the shared cache location given its name.
+	 **/
+	virtual FString GetSharedCachePath(const FString& CacheName) = 0;
+};
\ No newline at end of file
