#ue4 BP: class 'UObject' and Bluepirnt pins should not match

diff --git a/Engine/Config/BaseEditor.ini b/Engine/Config/BaseEditor.ini
--- a/Engine/Config/BaseEditor.ini
+++ b/Engine/Config/BaseEditor.ini
@@ -250,4 +250,4 @@ SupportWebsiteURL="http://unrealengine.com/support"
 bUseUserDefinedStructure=false
 
 [EditoronlyBP]
-bAllowClassObjAndBluepirntPinMatching=true
\ No newline at end of file
+bAllowClassObjAndBluepirntPinMatching=false
\ No newline at end of file
