#ue4 TTP 331831  Remove Morpheus as a default plugin until we work out messaging issues when it can't be found.

diff --git a/Engine/Config/BaseEngine.ini b/Engine/Config/BaseEngine.ini
--- a/Engine/Config/BaseEngine.ini
+++ b/Engine/Config/BaseEngine.ini
@@ -1991,7 +1991,6 @@ DebugComponentKey=Quote
 +EnabledPlugins="PluginsEditor"
 +EnabledPlugins="EpicSurvey"
 +EnabledPlugins="OculusRift"
-+EnabledPlugins="Morpheus"
 +EnabledPlugins="CustomMeshComponent"
 +EnabledPlugins="MessagingDebugger"
 +EnabledPlugins="PerforceSourceControl"
