#UE4 - Increase default connection timeout to support longer load times

diff --git a/Engine/Config/BaseEngine.ini b/Engine/Config/BaseEngine.ini
--- a/Engine/Config/BaseEngine.ini
+++ b/Engine/Config/BaseEngine.ini
@@ -1199,8 +1199,8 @@ ConfiguredLanSpeed=20000
 [/Script/OnlineSubsystemUtils.IpNetDriver]
 AllowPeerConnections=False
 AllowPeerVoice=False
-ConnectionTimeout=30.0
-InitialConnectTimeout=30.0
+ConnectionTimeout=60.0
+InitialConnectTimeout=60.0
 AckTimeout=1.0
 KeepAliveTime=0.2
 MaxClientRate=15000
