Launcher: L10N: The Launcher now properly looks in the correct directories for it's translation files.

diff --git a/Engine/Config/BaseGame.ini b/Engine/Config/BaseGame.ini
--- a/Engine/Config/BaseGame.ini
+++ b/Engine/Config/BaseGame.ini
@@ -1,5 +1,5 @@
 [Internationalization]
-+LocalizationPaths=%GAMEDIR%Content/Localization
++LocalizationPaths=%GAMEDIR%Content/Localization/Game
 
 [/Script/Engine.GameMode]
 
