#ue4 BaseEditor.ini: [CustomBlueprintFunctionLibrary] and [EditoronlyBP] sections restored

diff --git a/Engine/Config/BaseEditor.ini b/Engine/Config/BaseEditor.ini
--- a/Engine/Config/BaseEditor.ini
+++ b/Engine/Config/BaseEditor.ini
@@ -257,3 +257,12 @@ PulseDuration=5
 CdnUrl="http://cdn.unrealengine.com/newsfeed/{Culture}/index.json"
 EngineContentRelativeDirectory="Editor/NewsFeed/{Culture}/"
 Source="CDN"
+
+[CustomBlueprintFunctionLibrary]
+bCanCreateNew=true
+
+[EditoronlyBP]
+bAllowClassAndBlueprintPinMatching=false
+bReplaceBlueprintWithClass=false
+bDontLoadBlueprintOutsideEditor=false
+bBlueprintIsNotBlueprintType=false
\ No newline at end of file
