Fixed a bug causing editor framerate to be capped at 62 fps instead of 120

diff --git a/Engine/Config/BaseEngine.ini b/Engine/Config/BaseEngine.ini
--- a/Engine/Config/BaseEngine.ini
+++ b/Engine/Config/BaseEngine.ini
@@ -1264,8 +1264,7 @@ bCustomCameraAlignEmitter=true
 CustomCameraAlignEmitterDistance=100.0
 bDrawSocketsInGMode=false
 bSmoothFrameRate=false
-MinSmoothedFrameRate=5
-MaxSmoothedFrameRate=120
+SmoothedFrameRateRange=(LowerBound=(Type="ERangeBoundTypes::Inclusive",Value=5),UpperBound=(Type="ERangeBoundTypes::Inclusive",Value=120))
 UseOldStyleMICEditorGroups=true
 InEditorGameURLOptions=
 bOnlinePIEEnabled=false
