Removed OLD properties from device profile ini

diff --git a/Engine/Config/BaseDeviceProfiles.ini b/Engine/Config/BaseDeviceProfiles.ini
--- a/Engine/Config/BaseDeviceProfiles.ini
+++ b/Engine/Config/BaseDeviceProfiles.ini
@@ -27,26 +27,18 @@
 [Windows DeviceProfile]
 DeviceType=Windows
 BaseProfileName=
-MeshLODSettings=
-TextureLODSettings=
 
 [WindowsNoEditor DeviceProfile]
 DeviceType=WindowsNoEditor
 BaseProfileName=Windows
-MeshLODSettings=
-TextureLODSettings=
 
 [WindowsServer DeviceProfile]
 DeviceType=WindowsServer
 BaseProfileName=Windows
-MeshLODSettings=
-TextureLODSettings=
 
 [IOS DeviceProfile]
 DeviceType=IOS
 BaseProfileName=
-MeshLODSettings=
-TextureLODSettings=
 +CVars=r.BloomQuality=0
 +CVars=r.DepthOfFieldQuality=0
 +CVars=r.LightShaftQuality=0
@@ -54,36 +46,26 @@ TextureLODSettings=
 [iPad2 DeviceProfile]
 DeviceType=IOS
 BaseProfileName=IOS
-MeshLODSettings=
-TextureLODSettings=
 +CVars=r.RenderTargetSwitchWorkaround=1
 
 [iPad3 DeviceProfile]
 DeviceType=IOS
 BaseProfileName=IOS
-MeshLODSettings=
-TextureLODSettings=
 +CVars=r.RenderTargetSwitchWorkaround=1
 
 [iPad4 DeviceProfile]
 DeviceType=IOS
 BaseProfileName=IOS
-MeshLODSettings=
-TextureLODSettings=
 +CVars=r.RenderTargetSwitchWorkaround=1
 
 [iPadAir DeviceProfile]
 DeviceType=IOS
 BaseProfileName=IOS
-MeshLODSettings=
-TextureLODSettings=
 +CVars=r.BloomQuality=1
 
 [iPadMini DeviceProfile]
 DeviceType=IOS
 BaseProfileName=IOS
-MeshLODSettings=
-TextureLODSettings=
 +CVars=r.RenderTargetSwitchWorkaround=1
 
 [iPadMini2 DeviceProfile]
@@ -93,29 +75,21 @@ BaseProfileName=iPadAir
 [iPhone4 DeviceProfile]
 DeviceType=IOS
 BaseProfileName=IOS
-MeshLODSettings=
-TextureLODSettings=
 +CVars=r.RenderTargetSwitchWorkaround=1
 
 [iPhone4S DeviceProfile]
 DeviceType=IOS
 BaseProfileName=IOS
-MeshLODSettings=
-TextureLODSettings=
 +CVars=r.RenderTargetSwitchWorkaround=1
 
 [iPhone5 DeviceProfile]
 DeviceType=IOS
 BaseProfileName=IOS
-MeshLODSettings=
-TextureLODSettings=
 +CVars=r.RenderTargetSwitchWorkaround=1
 
 [iPhone5S DeviceProfile]
 DeviceType=IOS
 BaseProfileName=IOS
-MeshLODSettings=
-TextureLODSettings=
 +CVars=r.MobileContentScaleFactor=2
 +CVars=r.BloomQuality=1
 +CVars=r.DepthOfFieldQuality=1
@@ -124,22 +98,16 @@ TextureLODSettings=
 [IPodTouch4 DeviceProfile]
 DeviceType=IOS
 BaseProfileName=IOS
-MeshLODSettings=
-TextureLODSettings=
 +CVars=r.RenderTargetSwitchWorkaround=1
 
 [iPodTouch5 DeviceProfile]
 DeviceType=IOS
 BaseProfileName=IOS
-MeshLODSettings=
-TextureLODSettings=
 +CVars=r.RenderTargetSwitchWorkaround=1
 
 [Android DeviceProfile]
 DeviceType=Android
 BaseProfileName=
-MeshLODSettings=
-TextureLODSettings=
 +CVars=r.MobileContentScaleFactor=1
 +CVars=r.BloomQuality=0
 +CVars=r.DepthOfFieldQuality=0
@@ -187,14 +155,10 @@ BaseProfileName=Android_Adreno330
 [PS4 DeviceProfile]
 DeviceType=PS4
 BaseProfileName=
-MeshLODSettings=
-TextureLODSettings=
 
 [XboxOne DeviceProfile]
 DeviceType=XboxOne
 BaseProfileName=
-MeshLODSettings=
-TextureLODSettings=
 ; we output 10:10:10, not 8:8:8 so we don't need color quantization
 +CVars=r.TonemapperQuality=0
 ; For SSAO we rely on TemporalAA (with a randomized sample pattern over time) so we can use less samples
@@ -208,29 +172,19 @@ TextureLODSettings=
 [HTML5 DeviceProfile]
 DeviceType=HTML5
 BaseProfileName=
-MeshLODSettings=
-TextureLODSettings=
 
 [Mac DeviceProfile]
 DeviceType=Mac
 BaseProfileName=
-MeshLODSettings=
-TextureLODSettings=
 
 [MacNoEditor DeviceProfile]
 DeviceType=MacNoEditor
 BaseProfileName=Mac
-MeshLODSettings=
-TextureLODSettings=
 
 [MacServer DeviceProfile]
 DeviceType=MacServer
 BaseProfileName=Mac
-MeshLODSettings=
-TextureLODSettings=
 
 [WinRT DeviceProfile]
 DeviceType=WinRT
 BaseProfileName=
-MeshLODSettings=
-TextureLODSettings=
\ No newline at end of file
