SlateCore: added redirectors for UClasses that moved from Slate into SlateCore

diff --git a/Engine/Config/BaseEngine.ini b/Engine/Config/BaseEngine.ini
--- a/Engine/Config/BaseEngine.ini
+++ b/Engine/Config/BaseEngine.ini
@@ -1091,6 +1091,11 @@ bUseStreamingPause=false
 +TaggedPropertyRedirects=(ClassName="VehicleEngineData",OldPropertyName="MaxOmega",NewPropertyName="MaxRPM")
 +TaggedPropertyRedirects=(ClassName="VehicleMovementComponentBoosted4w",OldPropertyName="BoostedMaxOmega",NewPropertyName="BoostedMaxRPM")
 
+;SlateCore refactor
++ActiveClassRedirects=(OldClassName="SlateWidgetStyleAsset", NewClassName="/Script/SlateCore.SlateWidgetStyleAsset")
++ActiveClassRedirects=(OldClassName="SlateWidgetStyleContainerBase", NewClassName="/Script/SlateCore.SlateWidgetStyleContainerBase")
+
+
 [CoreUObject.Metadata]
 ; Note: UnrealHeaderTool should be rerun after making changes to MetadataRedirects to catch any keys specified in class headers
 +MetadataRedirects=(OldKey="K2Protected", NewKey="BlueprintProtected")
