Eliminate duplicate ssSwapControllers in console auto-complete list
TTP #334091

diff --git a/Engine/Config/BaseInput.ini b/Engine/Config/BaseInput.ini
--- a/Engine/Config/BaseInput.ini
+++ b/Engine/Config/BaseInput.ini
@@ -49,7 +49,6 @@ HistoryBot=-1
 +ManualAutoCompleteList=(Command="Exit",Desc="Exit (Exits the game)")
 +ManualAutoCompleteList=(Command="DebugCreatePlayer 1",Desc=)
 +ManualAutoCompleteList=(Command="FreezeAt",Desc="Locks the player view and rendering time.")
-+ManualAutoCompleteList=(Command="SSSwapControllers",Desc=)
 +ManualAutoCompleteList=(Command="Open",Desc="Open <MapName> (Opens the specified map, doesn't pass previously set options)")
 +ManualAutoCompleteList=(Command="Travel",Desc="Travel <MapName> (Travels to the specified map, passes along previously set options)")
 +ManualAutoCompleteList=(Command="ServerTravel",Desc="ServerTravel <MapName> (Travels to the specified map and brings clients along, passes along previously set options)")
