fixed bad INI comment (cannot be inline)

diff --git a/Engine/Config/BaseEngine.ini b/Engine/Config/BaseEngine.ini
--- a/Engine/Config/BaseEngine.ini
+++ b/Engine/Config/BaseEngine.ini
@@ -1893,7 +1893,8 @@ AgentMaxHeight=160.f
 AgentMaxStepHeight=35.f
 AgentMaxSlope=44.f
 MinRegionArea=0.f
-MergeRegionSize=400.f ; default should be aproximately 20*CellSize
+; default should be aproximately 20*CellSize
+MergeRegionSize=400.f
 bUseBetterOffsetsFromCorners=true
 
 
