fixed BaseEditor.ini settings that were accidentally removed during a bad merge

diff --git a/Engine/Config/BaseEditor.ini b/Engine/Config/BaseEditor.ini
--- a/Engine/Config/BaseEditor.ini
+++ b/Engine/Config/BaseEditor.ini
@@ -1,5 +1,8 @@
 [Internationalization]
 +LocalizationPaths=../../../Engine/Content/Localization/Editor
++ToolTipLocalizationPaths=../../../Engine/Content/Localization/Editor/ToolTips
++PropertyNameLocalizationPaths=../../../Engine/Content/Localization/Editor/PropertyNames
+ShouldLoadLocalizedPropertyNames=True
 
 [UnrealEd.GeneralSettings]
 EnableStylusPressure=True
@@ -96,7 +99,9 @@ DefaultProfile="StarterProfile"
 +StarterProfile=(NodeName="/Script/BlueprintGraph.K2Node_CallFunction", FieldName="/Script/Engine.GameplayStatics:GetAllActorsOfClass")
 +StarterProfile=(NodeName="/Script/BlueprintGraph.K2Node_SpawnActorFromClass")
 
-
+[UnrealEd.MeshPaint]
+MinBrushRadius=1
+MaxBrushRadius=256
 
 [UnrealEd.HitProxy]
 HitProxySize=5
@@ -260,4 +265,4 @@ bCanCreateNew=true
 bAllowClassAndBlueprintPinMatching=false
 bReplaceBlueprintWithClass=false
 bDontLoadBlueprintOutsideEditor=false
-bBlueprintIsNotBlueprintType=false
\ No newline at end of file
+bBlueprintIsNotBlueprintType=false
