Removed old crusty ini sections

diff --git a/Engine/Config/BaseEngine.ini b/Engine/Config/BaseEngine.ini
--- a/Engine/Config/BaseEngine.ini
+++ b/Engine/Config/BaseEngine.ini
@@ -1247,18 +1247,6 @@ AllowStreamingLightmaps=True
 UseDynamicStreaming=True
 BoostPlayerTextures=3.0
 
-[/Script/Engine.UIInteraction]
-UIJoystickDeadZone=0.9
-AxisRepeatDelay=0.2
-MouseButtonRepeatDelay=0.15
-DoubleClickTriggerSeconds=0.5
-DoubleClickPixelTolerance=1
-
-[/Script/Engine.GameUISceneClient]
-bCaptureUnprocessedInput=True
-bEnableDebugInput=true
-bRenderDebugInfo=false
-
 [/Script/UnrealEd.EditorEngine]
 AutoSaveDir=%GAMEDIR%Saved/Autosaves
 LocalPlayerClassName=/Script/Engine.LocalPlayer
