Broken test removed

diff --git a/Engine/Config/BaseEngine.ini b/Engine/Config/BaseEngine.ini
--- a/Engine/Config/BaseEngine.ini
+++ b/Engine/Config/BaseEngine.ini
@@ -1833,7 +1833,6 @@ ImportTestPackagePath=/Engine/Content/EditorAutomation
 
 [AutomationTesting.Blueprint]
 TestAllBlueprints=false
-+InstanceTestMaps=../../../Engine/Content/Maps/Automation/BlueprintInstanceTest.umap
 +ReparentTest.ChildrenPackagePaths=/Game/ReparentingTestAssets/Children
 +ReparentTest.ParentsPackagePaths=/Game/ReparentingTestAssets/Parents
 
