TTP 334628: Size By Speed for CPU particles
Works exactly the same as GPU version and runs through the same module.

diff --git a/Engine/Config/BaseEditor.ini b/Engine/Config/BaseEditor.ini
--- a/Engine/Config/BaseEditor.ini
+++ b/Engine/Config/BaseEditor.ini
@@ -213,7 +213,7 @@ bUseGrayScale=false
 +ModuleMenu_TypeDataToBaseModuleRejections=(ObjName="ParticleModuleTypeDataMesh",InvalidObjNames=("ParticleModuleBeamBase","ParticleModuleTrailBase","ParticleModuleVectorFieldBase"))
 +ModuleMenu_TypeDataToBaseModuleRejections=(ObjName="ParticleModuleTypeDataRibbon",InvalidObjNames=("ParticleModuleVectorFieldBase","ParticleModuleLightBase"))
 +ModuleMenu_TypeDataToBaseModuleRejections=(ObjName="ParticleModuleTypeDataGpu",InvalidObjNames=("ParticleModuleBeamBase","ParticleModuleTrailBase","ParticleModuleCameraBase","ParticleModuleKillBase","ParticleModuleParameterBase","ParticleModuleEventBase","ParticleModuleLightBase"))
-+ModuleMenu_TypeDataToSpecificModuleRejections=(ObjName="None",InvalidObjNames=("ParticleModuleCollisionGPU","ParticleModuleSizeScaleBySpeed","ParticleModuleMeshMaterial","ParticleModuleMeshRotation","ParticleModuleMeshRotation_Seeded","ParticleModuleMeshRotationRate","ParticleModuleMeshRotationRate_Seeded","ParticleModuleMeshRotationRateMultiplyLife","ParticleModuleMeshRotationRateOverLife","ParticleModuleTrailSpawnPerUnit", "ParticleModuleAttractorPointGravity"))
++ModuleMenu_TypeDataToSpecificModuleRejections=(ObjName="None",InvalidObjNames=("ParticleModuleCollisionGPU","ParticleModuleMeshMaterial","ParticleModuleMeshRotation","ParticleModuleMeshRotation_Seeded","ParticleModuleMeshRotationRate","ParticleModuleMeshRotationRate_Seeded","ParticleModuleMeshRotationRateMultiplyLife","ParticleModuleMeshRotationRateOverLife","ParticleModuleTrailSpawnPerUnit", "ParticleModuleAttractorPointGravity"))
 +ModuleMenu_TypeDataToSpecificModuleRejections=(ObjName="ParticleModuleTypeDataBeam2",InvalidObjNames=("ParticleModuleSizeScaleBySpeed","ParticleModuleMeshMaterial","ParticleModuleMeshRotation","ParticleModuleMeshRotation_Seeded","ParticleModuleMeshRotationRate","ParticleModuleMeshRotationRate_Seeded","ParticleModuleMeshRotationRateMultiplyLife","ParticleModuleMeshRotationRateOverLife","ParticleModuleTrailSpawnPerUnit","ParticleModuleAccelerationDrag"))
 +ModuleMenu_TypeDataToSpecificModuleRejections=(ObjName="ParticleModuleTypeDataRibbon",InvalidObjNames=("ParticleModuleSizeScaleBySpeed","ParticleModuleMeshMaterial","ParticleModuleMeshRotation","ParticleModuleMeshRotation_Seeded","ParticleModuleMeshRotationRate","ParticleModuleMeshRotationRate_Seeded","ParticleModuleMeshRotationRateMultiplyLife","ParticleModuleMeshRotationRateOverLife","ParticleModuleLocationEmitterDirect" ))
 +ModuleMenu_TypeDataToSpecificModuleRejections=(ObjName="ParticleModuleTypeDataGpu",InvalidObjNames=("ParticleModuleCollision","ParticleModuleAttractorPoint","ParticleModuleAttractorLine","ParticleModuleAttractorParticle","ParticleModuleAcceleration","ParticleModuleAccelerationOverLifetime","ParticleModuleLocationEmitterDirect","ParticleModuleLocationDirect","ParticleModuleVelocityOverLifetime","ParticleModuleSubUVMovie","ParticleModuleRotationRateMultiplyLife","ParticleModuleRotationOverLifetime","ParticleModuleSourceMovement","ParticleModuleEventGenerator","ParticleModuleMeshMaterial","ParticleModuleMeshRotation","ParticleModuleMeshRotation_Seeded","ParticleModuleMeshRotationRate","ParticleModuleMeshRotationRate_Seeded","ParticleModuleMeshRotationRateMultiplyLife","ParticleModuleMeshRotationRateOverLife","ParticleModuleTrailSpawnPerUnit"))
diff --git a/Engine/Source/Runtime/Engine/Classes/Particles/Size/ParticleModuleSizeScaleBySpeed.h b/Engine/Source/Runtime/Engine/Classes/Particles/Size/ParticleModuleSizeScaleBySpeed.h
--- a/Engine/Source/Runtime/Engine/Classes/Particles/Size/ParticleModuleSizeScaleBySpeed.h
+++ b/Engine/Source/Runtime/Engine/Classes/Particles/Size/ParticleModuleSizeScaleBySpeed.h
@@ -23,7 +23,8 @@ class UParticleModuleSizeScaleBySpeed : public UParticleModuleSizeBase
 
 
 	// Begin UParticleModule Interface
-	virtual void CompileModule( struct FParticleEmitterBuildInfo& EmitterInfo ) OVERRIDE;
+	virtual void Update(FParticleEmitterInstance* Owner, int32 Offset, float DeltaTime) OVERRIDE;
+	virtual void CompileModule(struct FParticleEmitterBuildInfo& EmitterInfo) OVERRIDE;
 	// End UParticleModule Interface
 };
 
diff --git a/Engine/Source/Runtime/Engine/Private/Particles/ParticleModules_Size.cpp b/Engine/Source/Runtime/Engine/Private/Particles/ParticleModules_Size.cpp
--- a/Engine/Source/Runtime/Engine/Private/Particles/ParticleModules_Size.cpp
+++ b/Engine/Source/Runtime/Engine/Private/Particles/ParticleModules_Size.cpp
@@ -432,10 +432,24 @@ void UParticleModuleSizeScale::SetToSensibleDefaults(UParticleEmitter* Owner)
 UParticleModuleSizeScaleBySpeed::UParticleModuleSizeScaleBySpeed(const class FPostConstructInitializeProperties& PCIP)
 	: Super(PCIP)
 {
+	bUpdateModule = true;
 	MaxScale.X = 1;
 	MaxScale.Y = 1;
 }
 
+void UParticleModuleSizeScaleBySpeed::Update(FParticleEmitterInstance* Owner, int32 Offset, float DeltaTime)
+{
+	FVector Scale(SpeedScale.X, SpeedScale.Y, 1.0f);
+	FVector ScaleMax(MaxScale.X, MaxScale.Y, 1.0f);
+
+	BEGIN_UPDATE_LOOP;
+		FVector Size = Scale * Particle.Velocity.Size();
+		Size = Size.ComponentMax(FVector(1.0f));
+		Size = Size.ComponentMin(ScaleMax);
+		Particle.Size = Particle.BaseSize * Size;
+	END_UPDATE_LOOP;
+}
+
 
 /**
  * Compile the effects of this module on runtime simulation.
