Changing the name of UPointLightComponent::SetRadius to SetAttenuationRadius to match the name change of the variable UPointLightComponent::Radius to AttenuationRadius.
Including a redirect for old Blueprint uses to the new one in BaseEngine.ini
Didn't change UPointLight::SetRadius because that function is deprecated anyway.

diff --git a/Engine/Config/BaseEngine.ini b/Engine/Config/BaseEngine.ini
--- a/Engine/Config/BaseEngine.ini
+++ b/Engine/Config/BaseEngine.ini
@@ -1206,6 +1206,7 @@ bUseStreamingPause=false
 ; VER_UE4_INVERSE_SQUARED_LIGHTS_DEFAULT
 +K2FieldRedirects=(OldFieldName="LightComponentBase.Brightness",NewFieldName="LightComponentBase.Intensity")
 +K2FieldRedirects=(OldFieldName="PointLightComponent.Radius",NewFieldName="PointLightComponent.AttenuationRadius")
++K2FieldRedirects=(OldFieldName="PointLightComponent.SetRadius",NewFieldName="PointLightComponent.SetAttenuationRadius")
 +K2FieldRedirects=(OldFieldName="LightComponent.InverseSquaredFalloff",NewFieldName="PointLightComponent.bUseInverseSquaredFalloff")
 
 ; Vehicle rename properties
diff --git a/Engine/Source/Runtime/Engine/Classes/Components/PointLightComponent.h b/Engine/Source/Runtime/Engine/Classes/Components/PointLightComponent.h
--- a/Engine/Source/Runtime/Engine/Classes/Components/PointLightComponent.h
+++ b/Engine/Source/Runtime/Engine/Classes/Components/PointLightComponent.h
@@ -58,7 +58,7 @@ class ENGINE_API UPointLightComponent : public ULightComponent
 	struct FLightmassPointLightSettings LightmassSettings;
 
 	UFUNCTION(BlueprintCallable, Category="Rendering|Lighting")
-	void SetRadius(float NewRadius);
+	void SetAttenuationRadius(float NewRadius);
 
 	UFUNCTION(BlueprintCallable, Category="Rendering|Lighting")
 	void SetLightFalloffExponent(float NewLightFalloffExponent);
diff --git a/Engine/Source/Runtime/Engine/Private/Light.cpp b/Engine/Source/Runtime/Engine/Private/Light.cpp
--- a/Engine/Source/Runtime/Engine/Private/Light.cpp
+++ b/Engine/Source/Runtime/Engine/Private/Light.cpp
@@ -304,7 +304,7 @@ void ADirectionalLight::PostEditChangeProperty(FPropertyChangedEvent& PropertyCh
 
 void APointLight::SetRadius(float NewRadius)
 {
-	PointLightComponent->SetRadius(NewRadius);
+	PointLightComponent->SetAttenuationRadius(NewRadius);
 }
 
 void APointLight::SetLightFalloffExponent(float NewLightFalloffExponent)
diff --git a/Engine/Source/Runtime/Engine/Private/PointLightComponent.cpp b/Engine/Source/Runtime/Engine/Private/PointLightComponent.cpp
--- a/Engine/Source/Runtime/Engine/Private/PointLightComponent.cpp
+++ b/Engine/Source/Runtime/Engine/Private/PointLightComponent.cpp
@@ -126,7 +126,7 @@ FLightSceneProxy* UPointLightComponent::CreateSceneProxy() const
 	return new FPointLightSceneProxy(this);
 }
 
-void UPointLightComponent::SetRadius(float NewRadius)
+void UPointLightComponent::SetAttenuationRadius(float NewRadius)
 {
 	// Only movable lights can change their radius at runtime
 	if (!(IsRegistered() && (Mobility == EComponentMobility::Static || Mobility == EComponentMobility::Stationary))
