UE4 - revert config change

diff --git a/Engine/Config/ConsoleVariables.ini b/Engine/Config/ConsoleVariables.ini
--- a/Engine/Config/ConsoleVariables.ini
+++ b/Engine/Config/ConsoleVariables.ini
@@ -20,5 +20,6 @@
 [Startup]
 
 ; Uncomment to get detailed logs on shader compiles and the opportunity to retry on errors
-r.ShaderDevelopmentMode=1
-r.DumpShaderDebugInfo=1
\ No newline at end of file
+;r.ShaderDevelopmentMode=1
+; Uncomment to dump shaders in the Saved folder
+;r.DumpShaderDebugInfo=1
