Change Set Anim Class to Set Anim Instance Class to be consistent.

diff --git a/Engine/Config/BaseEngine.ini b/Engine/Config/BaseEngine.ini
--- a/Engine/Config/BaseEngine.ini
+++ b/Engine/Config/BaseEngine.ini
@@ -406,8 +406,9 @@ bUseStreamingPause=false
 +K2FieldRedirects=(OldFieldName="SingleAnimSkeletalComponent.GetPlayRate",NewFieldName="SkeletalMeshComponent.GetPlayRate")
 +K2FieldRedirects=(OldFieldName="SingleAnimSkeletalComponent.IsPlaying",NewFieldName="SkeletalMeshComponent.IsPlaying")
 
-+K2FieldRedirects=(OldFieldName="SkeletalMeshComponent.SetAnimBlueprint",NewFieldName="SkeletalMeshComponent.SetAnimClass")
-+K2ParamRedirects=(NodeName="/Script/BlueprintGraph.K2Node_CallFunction", OldParamName="SkeletalMeshComponent.SetAnimClass.NewBlueprint", NewParamName="NewClass")
++K2FieldRedirects=(OldFieldName="SkeletalMeshComponent.SetAnimBlueprint",NewFieldName="SkeletalMeshComponent.SetAnimInstanceClass")
++K2ParamRedirects=(NodeName="/Script/BlueprintGraph.K2Node_CallFunction", OldParamName="SkeletalMeshComponent.SetAnimInstanceClass.NewBlueprint", NewParamName="NewClass")
++K2FieldRedirects=(OldFieldName="SkeletalMeshComponent.SetAnimClass",NewFieldName="SkeletalMeshComponent.SetAnimInstanceClass")
 
 +ActiveClassRedirects=(OldClassName="DynamicBlockingVolume",NewClassName="BlockingVolume")
 +ActiveClassRedirects=(OldClassName="DynamicPhysicsVolume",NewClassName="PhysicsVolume")
diff --git a/Engine/Source/Editor/Persona/Private/Persona.cpp b/Engine/Source/Editor/Persona/Private/Persona.cpp
--- a/Engine/Source/Editor/Persona/Private/Persona.cpp
+++ b/Engine/Source/Editor/Persona/Private/Persona.cpp
@@ -849,7 +849,7 @@ void FPersona::InitPersona(const EToolkitMode::Type Mode, const TSharedPtr< clas
 		}
 	}
 
-	PreviewComponent->SetAnimClass(AnimBlueprint ? AnimBlueprint->GeneratedClass : NULL);
+	PreviewComponent->SetAnimInstanceClass(AnimBlueprint ? AnimBlueprint->GeneratedClass : NULL);
 
 	// We always want a preview instance unless we are using blueprints so that bone manipulation works
 	if (AnimBlueprint == NULL)
diff --git a/Engine/Source/Editor/UnrealEd/Private/Factories/ActorFactory.cpp b/Engine/Source/Editor/UnrealEd/Private/Factories/ActorFactory.cpp
--- a/Engine/Source/Editor/UnrealEd/Private/Factories/ActorFactory.cpp
+++ b/Engine/Source/Editor/UnrealEd/Private/Factories/ActorFactory.cpp
@@ -885,7 +885,7 @@ void UActorFactorySkeletalMesh::PostSpawnActor( UObject* Asset, AActor* NewActor
 	NewSMActor->SkeletalMeshComponent->RegisterComponent();
 	if( AnimBlueprint )
 	{
-		NewSMActor->SkeletalMeshComponent->SetAnimClass(AnimBlueprint->GeneratedClass);
+		NewSMActor->SkeletalMeshComponent->SetAnimInstanceClass(AnimBlueprint->GeneratedClass);
 	}
 }
 
diff --git a/Engine/Source/Editor/UnrealEd/Private/LevelEditorViewport.cpp b/Engine/Source/Editor/UnrealEd/Private/LevelEditorViewport.cpp
--- a/Engine/Source/Editor/UnrealEd/Private/LevelEditorViewport.cpp
+++ b/Engine/Source/Editor/UnrealEd/Private/LevelEditorViewport.cpp
@@ -647,7 +647,7 @@ static bool AttemptApplyObjToActor( UObject* ObjToUse, AActor* ActorToApplyTo, i
 						if(SkelMeshComponent->SkeletalMesh &&
 							NeedsSkeleton->IsCompatible(SkelMeshComponent->SkeletalMesh->Skeleton))
 						{
-							SkelMeshComponent->SetAnimClass(DroppedObjAsAnimBlueprint->GeneratedClass);
+							SkelMeshComponent->SetAnimInstanceClass(DroppedObjAsAnimBlueprint->GeneratedClass);
 							bResult = true;
 						}
 					}
diff --git a/Engine/Source/Runtime/Engine/Classes/Components/SkeletalMeshComponent.h b/Engine/Source/Runtime/Engine/Classes/Components/SkeletalMeshComponent.h
--- a/Engine/Source/Runtime/Engine/Classes/Components/SkeletalMeshComponent.h
+++ b/Engine/Source/Runtime/Engine/Classes/Components/SkeletalMeshComponent.h
@@ -238,7 +238,7 @@ public:
 
 	/** 
 	 * The AnimBlueprint class to use
-	 * Use 'SetAnimClass' to change at runtime. 
+	 * Use 'SetAnimInstanceClass' to change at runtime. 
 	 */
 	UPROPERTY(EditAnywhere, BlueprintReadOnly, Category = Animation)
 	class UAnimBlueprintGeneratedClass* AnimBlueprintGeneratedClass;
@@ -444,8 +444,8 @@ public:
 	UPROPERTY(BlueprintAssignable)
 	FConstraintBrokenSignature OnConstraintBroken;
 
-	UFUNCTION(BlueprintCallable, Category="Components|SkeletalMesh")
-	void SetAnimClass(class UClass* NewClass);
+	UFUNCTION(BlueprintCallable, Category="Components|SkeletalMesh", meta=(Keywords = "AnimBlueprint"))
+	void SetAnimInstanceClass(class UClass* NewClass);
 
 	/** 
 	 * Returns the animation instance that is driving the class (if available). This is typically an instance of
diff --git a/Engine/Source/Runtime/Engine/Private/SkeletalMeshComponent.cpp b/Engine/Source/Runtime/Engine/Private/SkeletalMeshComponent.cpp
--- a/Engine/Source/Runtime/Engine/Private/SkeletalMeshComponent.cpp
+++ b/Engine/Source/Runtime/Engine/Private/SkeletalMeshComponent.cpp
@@ -1133,7 +1133,7 @@ void USkeletalMeshComponent::SetForceRefPose(bool bNewForceRefPose)
 	MarkRenderStateDirty();
 }
 
-void USkeletalMeshComponent::SetAnimClass(class UClass* NewClass)
+void USkeletalMeshComponent::SetAnimInstanceClass(class UClass* NewClass)
 {
 	if (NewClass != NULL)
 	{
