MSAA is now reduced or tuned off for lower engine scalability levels.
MSAA levels are as follows: Low: 1, Medium: 1, High: 2, Epic: 4.

Shader enhancements to improve the quality of selection outline are still to come.

diff --git a/Engine/Config/BaseScalability.ini b/Engine/Config/BaseScalability.ini
--- a/Engine/Config/BaseScalability.ini
+++ b/Engine/Config/BaseScalability.ini
@@ -9,6 +9,11 @@
 
 ; Note: [ResolutionQuality] isn't using the usual 0/1/2/3 quality levels, the value directly maps to r.ScreenPercentage.
 ; This is because we want custom scaling aligned with the screen/window resolution with more fine grained control
+;
+;-----------------------------------------------------------------------------------------------------------------
+;
+; Note: "r.MSAA.CompositingSampleCount" goes upto 8 however we're limiting it to 4 in a scale of 1/1/2/4
+; Meaning that it is off for the first two detail levels.
 
 ;-----------------------------------------------------------------------------------------------------------------
 
@@ -96,6 +101,7 @@ r.FastBlurThreshold=0
 r.UpsampleQuality=1
 r.TonemapperQuality=0
 r.LightShaftQuality=0
+r.MSAA.CompositingSampleCount=1
 
 [PostProcessQuality@1]
 r.MotionBlurQuality=3
@@ -111,6 +117,7 @@ r.FastBlurThreshold=2
 r.UpsampleQuality=2
 r.TonemapperQuality=0
 r.LightShaftQuality=0
+r.MSAA.CompositingSampleCount=1
 
 [PostProcessQuality@2]
 r.MotionBlurQuality=3
@@ -126,6 +133,7 @@ r.FastBlurThreshold=3
 r.UpsampleQuality=2
 r.TonemapperQuality=1
 r.LightShaftQuality=1
+r.MSAA.CompositingSampleCount=2
 
 [PostProcessQuality@3]
 r.MotionBlurQuality=4
@@ -141,6 +149,7 @@ r.FastBlurThreshold=7
 r.UpsampleQuality=3
 r.TonemapperQuality=1
 r.LightShaftQuality=1
+r.MSAA.CompositingSampleCount=4
 
 ;-----------------------------------------------------------------------------------------------------------------
 
