Paper2D: Enable Paper2D plugin by default so that content-only games can use it in cooked builds

diff --git a/Engine/Config/BaseEngine.ini b/Engine/Config/BaseEngine.ini
--- a/Engine/Config/BaseEngine.ini
+++ b/Engine/Config/BaseEngine.ini
@@ -2164,6 +2164,7 @@ DebugComponentHUDClassName="/Script/GameplayDebugger.GameplayDebuggingHUDCompone
 +EnabledPlugins="SlateRemote"
 +EnabledPlugins="ScriptPlugin"
 +EnabledPlugins="ScriptEditorPlugin"
++EnabledPlugins="Paper2D"
 +EnabledPlugins="SpeedTreeImporter"
 
 ; Note: These should be in IOSEngine, but are placed here because the Settings editor UI cannot write to the Game IOSEngine, meaning settings in Engine/IOSEngine are non-overridable
