Remove merge markers.

diff --git a/Engine/Config/BaseEngine.ini b/Engine/Config/BaseEngine.ini
--- a/Engine/Config/BaseEngine.ini
+++ b/Engine/Config/BaseEngine.ini
@@ -2144,55 +2144,6 @@ StartingBeaconPort=15000
 DebugComponentClassName="/Script/GameplayDebugger.GameplayDebuggingComponent"
 DebugComponentHUDClassName="/Script/GameplayDebugger.GameplayDebuggingHUDComponent"
 
->>>> ORIGINAL //depot/UE4/Engine/Config/BaseEngine.ini#696
-[Plugins]
-+EnabledPlugins="BlankPlugin"
-+EnabledPlugins="UObjectPlugin"
-+EnabledPlugins="PluginsEditor"
-+EnabledPlugins="EpicSurvey"
-+EnabledPlugins="OculusRift"
-+EnabledPlugins="CustomMeshComponent"
-+EnabledPlugins="MessagingDebugger"
-+EnabledPlugins="PerforceSourceControl"
-+EnabledPlugins="SubversionSourceControl"
-+EnabledPlugins="UdpMessaging"
-+EnabledPlugins="WindowsMoviePlayer"
-+EnabledPlugins="CableComponent"
-+EnabledPlugins="ExampleDeviceProfileSelector"
-+EnabledPlugins="WinDualShock"
-+EnabledPlugins="VisualStudioSourceCodeAccess"
-+EnabledPlugins="XCodeSourceCodeAccess"
-+EnabledPlugins="SlateRemote"
-+EnabledPlugins="ScriptPlugin"
-+EnabledPlugins="ScriptEditorPlugin"
-+EnabledPlugins="SpeedTreeImporter"
-
-==== THEIRS //depot/UE4/Engine/Config/BaseEngine.ini#697
-[Plugins]
-+EnabledPlugins="BlankPlugin"
-+EnabledPlugins="UObjectPlugin"
-+EnabledPlugins="PluginsEditor"
-+EnabledPlugins="EpicSurvey"
-+EnabledPlugins="OculusRift"
-+EnabledPlugins="CustomMeshComponent"
-+EnabledPlugins="MessagingDebugger"
-+EnabledPlugins="PerforceSourceControl"
-+EnabledPlugins="SubversionSourceControl"
-+EnabledPlugins="UdpMessaging"
-+EnabledPlugins="WindowsMoviePlayer"
-+EnabledPlugins="CableComponent"
-+EnabledPlugins="ExampleDeviceProfileSelector"
-+EnabledPlugins="WinDualShock"
-+EnabledPlugins="VisualStudioSourceCodeAccess"
-+EnabledPlugins="XCodeSourceCodeAccess"
-+EnabledPlugins="SlateRemote"
-+EnabledPlugins="ScriptPlugin"
-+EnabledPlugins="ScriptEditorPlugin"
-+EnabledPlugins="Paper2D"
-+EnabledPlugins="SpeedTreeImporter"
-
-==== YOURS //Ben.Marsh_T3245/UE4/Engine/Config/BaseEngine.ini
-<<<<
 ; Note: These should be in IOSEngine, but are placed here because the Settings editor UI cannot write to the Game IOSEngine, meaning settings in Engine/IOSEngine are non-overridable
 [/Script/IOSRuntimeSettings.IOSRuntimeSettings]
 bEnableGameCenterSupport=False
