* fixed Materiee not finding old PP settings
  TTP 340410 SAMPLES: Infiltrator: Matinee track warning for Exposure Offset

diff --git a/Engine/Config/BaseEngine.ini b/Engine/Config/BaseEngine.ini
--- a/Engine/Config/BaseEngine.ini
+++ b/Engine/Config/BaseEngine.ini
@@ -1139,6 +1139,14 @@ bUseStreamingPause=false
 +MatineeTrackRedirects=(TargetClassName="/Script/Engine.CameraActor",OldFieldName="PostProcessSettings.EyeAdaptionSpeedUp",NewFieldName="CameraComponent.PostProcessSettings.AutoExposureSpeedUp")
 +MatineeTrackRedirects=(TargetClassName="/Script/Engine.CameraActor",OldFieldName="PostProcessSettings.ExposureOffset",NewFieldName="CameraComponent.PostProcessSettings.AutoExposureBias")
 
++MatineeTrackRedirects=(TargetClassName="/Script/Engine.PostprocessVolume",OldFieldName="Settings.EyeAdaptationLowPercent",NewFieldName="Settings.AutoExposureLowPercent")
++MatineeTrackRedirects=(TargetClassName="/Script/Engine.PostprocessVolume",OldFieldName="Settings.EyeAdaptationHighPercent",NewFieldName="Settings.AutoExposureHighPercent")
++MatineeTrackRedirects=(TargetClassName="/Script/Engine.PostprocessVolume",OldFieldName="Settings.EyeAdaptationMinBrightness",NewFieldName="Settings.AutoExposureMinBrightness")
++MatineeTrackRedirects=(TargetClassName="/Script/Engine.PostprocessVolume",OldFieldName="Settings.EyeAdaptationMaxBrightness",NewFieldName="Settings.AutoExposureMaxBrightness")
++MatineeTrackRedirects=(TargetClassName="/Script/Engine.PostprocessVolume",OldFieldName="Settings.EyeAdaptionSpeedDown",NewFieldName="Settings.AutoExposureSpeedDown")
++MatineeTrackRedirects=(TargetClassName="/Script/Engine.PostprocessVolume",OldFieldName="Settings.EyeAdaptionSpeedUp",NewFieldName="Settings.AutoExposureSpeedUp")
++MatineeTrackRedirects=(TargetClassName="/Script/Engine.PostprocessVolume",OldFieldName="Settings.ExposureOffset",NewFieldName="Settings.AutoExposureBias")
+
 +MatineeTrackRedirects=(TargetClassName="/Script/Engine.CameraActor",OldFieldName="FOVAngle",NewFieldName="CameraComponent.FieldOfView")
 +MatineeTrackRedirects=(TargetClassName="/Script/Engine.CameraActor",OldFieldName="AspectRatio",NewFieldName="CameraComponent.AspectRatio")
 +MatineeTrackRedirects=(TargetClassName="/Script/Engine.CameraActor",OldFieldName="bConstrainAspectRatio",NewFieldName="CameraComponent.bConstrainAspectRatio")
