CrashReport - Updated configuration

diff --git a/Engine/Config/BaseEngine.ini b/Engine/Config/BaseEngine.ini
--- a/Engine/Config/BaseEngine.ini
+++ b/Engine/Config/BaseEngine.ini
@@ -2059,13 +2059,13 @@ AlwaysRecordingStats=true
 ; ini files do not like '//' - they treat it as a comment
 LocalSymbolStore=d:/Symbols/UE4
 SourceControlBuildLabelPattern=*-CL-%CHANGELISTNUMBER%
+ExecutablePathPattern=\\epicgames.net\root\Builds\Rocket\Automated\%PRODUCT_VERSION%\Windows
 ; devweb-02
 PDBCachePath=F:/CrashReportPDBCache/
-;PDBCachePath=U:/CrashReport/PDBCache/
+DepotRoot=F:/depot
 DaysToDeleteUnusedFilesFromPDBCache=14
 PDBCacheSizeGB=128
 MinDiskFreeSpaceGB=64
-bUsePDBCache=true
 
 [Engine.ErrorHandling]
 bPromptForRemoteDebugging=false
