Add redirector for deleted EmitterSpawnable class

diff --git a/Engine/Config/BaseEngine.ini b/Engine/Config/BaseEngine.ini
--- a/Engine/Config/BaseEngine.ini
+++ b/Engine/Config/BaseEngine.ini
@@ -1260,6 +1260,9 @@ bUseStreamingPause=false
 ;Foliage refactor
 +ActiveClassRedirects=(OldClassName="InstancedFoliageSettings", NewClassName="FoliageType_InstancedStaticMesh")
 
+;Removal of EmitterSpawnable
++ActiveClassRedirects=(OldClassName="EmitterSpawnable", NewClassName="Emitter")
+
 
 [CoreUObject.Metadata]
 ; Note: UnrealHeaderTool should be rerun after making changes to MetadataRedirects to catch any keys specified in class headers
