Misc - Feature Request: EDITOR: Expose stat unitgraph in command line

diff --git a/Engine/Config/BaseInput.ini b/Engine/Config/BaseInput.ini
--- a/Engine/Config/BaseInput.ini
+++ b/Engine/Config/BaseInput.ini
@@ -86,6 +86,7 @@ MaxScrollbackSize=1024
 +ManualAutoCompleteList=(Command="Show SPLINES",Desc="Show SPLINES (Toggles spline rendering)")
 +ManualAutoCompleteList=(Command="Stat FPS",Desc="Stat FPS (Shows FPS counter)")
 +ManualAutoCompleteList=(Command="Stat UNIT",Desc="Stat UNIT (Shows hardware unit framerate)")
++ManualAutoCompleteList=(Command="Stat UnitGraph",Desc="Stat UnitGraph (Draws simple unit time graph)")
 +ManualAutoCompleteList=(Command="Stat DUMPHITCHES",Desc="Stat DumpHitches (executes dumpstats on hitches - see log)")
 +ManualAutoCompleteList=(Command="Stat D3D11RHI",Desc="Stat D3D11RHI (Shows Direct3D 11 stats)")
 +ManualAutoCompleteList=(Command="Stat LEVELS",Desc="Stat LEVELS (Displays level streaming info)")
