Paper2D: Prevent external content reference warnings on placed assets that use the default plugin materials

diff --git a/Engine/Config/BaseEngine.ini b/Engine/Config/BaseEngine.ini
--- a/Engine/Config/BaseEngine.ini
+++ b/Engine/Config/BaseEngine.ini
@@ -1359,6 +1359,7 @@ SizeOfPermanentObjectPool=0
 AsyncIOBandwidthLimit=0
 +Paths=../../../Engine/Content
 +Paths=%GAMEDIR%Content
++Paths=../../../Engine/Plugins/Experimental/Paper2D/Content
 CutdownPaths=%GAMEDIR%CutdownPackages
 ZeroEngineVersionWarning=True
 UseStrictEngineVersioning=True
