UDN: INCORRECT FUNCTIONALITY: Format Text node causes compile errors unless the blueprint is recompiled after the level is reoà ttp#341425

https://udn.unrealengine.com/questions/213788/formattext-node-causes-compile-errors-in-42.html

diff --git a/Engine/Config/BaseEngine.ini b/Engine/Config/BaseEngine.ini
--- a/Engine/Config/BaseEngine.ini
+++ b/Engine/Config/BaseEngine.ini
@@ -1184,6 +1184,8 @@ bUseStreamingPause=false
 +ActiveClassRedirects=(OldClassName="AnimNode_BlendSpace",NewClassName="/Script/Engine.AnimNode_BlendSpacePlayer")
 +ActiveStructRedirects=(OldStructName="AnimNode_BlendSpace",NewStructName="AnimNode_BlendSpacePlayer")
 
++ActiveStructRedirects=(OldStructName="FFormatTextArgument",NewStructName="/Script/Engine.FFormatTextArgument")
+
 +K2FieldRedirects=(OldFieldName="AnimInstance.GetSequencePlayerLength",NewFieldName="GetAnimAssetPlayerLength")
 +K2FieldRedirects=(OldFieldName="AnimInstance.GetSequencePlayerTimeFraction",NewFieldName="GetAnimAssetPlayerTimeFraction")
 +K2FieldRedirects=(OldFieldName="AnimInstance.GetSequencePlayerTimeFromEnd",NewFieldName="GetAnimAssetPlayerTimeFromEnd")
