UE4 - fix shared DDC excessive load

diff --git a/Engine/Config/BaseEngine.ini b/Engine/Config/BaseEngine.ini
--- a/Engine/Config/BaseEngine.ini
+++ b/Engine/Config/BaseEngine.ini
@@ -1929,8 +1929,8 @@ AsyncPut=(Type=AsyncPut, Inner=Hierarchy)
 Hierarchy=(Type=Hierarchical, Inner=Boot, Inner=Pak, Inner=EnginePak, Inner=Local, Inner=Shared)
 Boot=(Type=Boot, Filename="%GAMEDIR%DerivedDataCache/Boot.ddc", MaxCacheSize=512)
 Local=(Type=FileSystem, ReadOnly=false, Clean=false, Flush=false, PurgeTransient=true, DeleteUnused=true, UnusedFileAge=34, FoldersToClean=-1, Path=../../../Engine/DerivedDataCache, EnvPathOverride=UE-LocalDataCachePath)
-Shared=(Type=FileSystem, ReadOnly=false, Clean=false, Flush=false, DeleteUnused=true, UnusedFileAge=23, FoldersToClean=-1, Path=?EpicDDC, EnvPathOverride=UE-SharedDataCachePath)
-AltShared=(Type=FileSystem, ReadOnly=true, Clean=false, Flush=false, DeleteUnused=true, UnusedFileAge=23, FoldersToClean=-1, Path=?EpicDDC2, EnvPathOverride=UE-SharedDataCachePath2)
+Shared=(Type=FileSystem, ReadOnly=false, Clean=false, Flush=false, DeleteUnused=true, UnusedFileAge=23, FoldersToClean=10, MaxFileChecksPerSec=1, Path=?EpicDDC, EnvPathOverride=UE-SharedDataCachePath)
+AltShared=(Type=FileSystem, ReadOnly=true, Clean=false, Flush=false, DeleteUnused=true, UnusedFileAge=23, FoldersToClean=10, MaxFileChecksPerSec=1, Path=?EpicDDC2, EnvPathOverride=UE-SharedDataCachePath2)
 Pak=(Type=ReadPak, Filename="%GAMEDIR%DerivedDataCache/DDC.ddp")
 EnginePak=(Type=ReadPak, Filename=../../../Engine/DerivedDataCache/DDC.ddp)
 
@@ -1941,7 +1941,7 @@ AsyncPut=(Type=AsyncPut, Inner=Hierarchy)
 Hierarchy=(Type=Hierarchical, Inner=Boot, Inner=Pak, Inner=EnginePak, Inner=Local, Inner=Seattle)
 Boot=(Type=Boot, Filename="%GAMEDIR%DerivedDataCache/Boot.ddc", MaxCacheSize=512)
 Local=(Type=FileSystem, ReadOnly=false, Clean=false, Flush=false, PurgeTransient=true, DeleteUnused=true, UnusedFileAge=34, FoldersToClean=-1, Path=../../../Engine/DerivedDataCache)
-Seattle=(Type=FileSystem, ReadOnly=false, Clean=false, Flush=false, DeleteUnused=true, UnusedFileAge=23, FoldersToClean=-1, Path=\\sea.epicgames.net\root\DDC, EnvPathOverride=UE-SharedDataCachePath_Seattle)
+Seattle=(Type=FileSystem, ReadOnly=false, Clean=false, Flush=false, DeleteUnused=true, UnusedFileAge=23, FoldersToClean=10, MaxFileChecksPerSec=1, Path=\\sea.epicgames.net\root\DDC, EnvPathOverride=UE-SharedDataCachePath_Seattle)
 Pak=(Type=ReadPak, Filename="%GAMEDIR%DerivedDataCache/DDC.ddp")
 EnginePak=(Type=ReadPak, Filename=../../../Engine/DerivedDataCache/DDC.ddp)
 
@@ -1971,8 +1971,8 @@ AsyncPut=(Type=AsyncPut, Inner=Hierarchy)
 Hierarchy=(Type=Hierarchical, Inner=Boot, Inner=PakWrite, Inner=PakRead, Inner=Local, Inner=Shared)
 Boot=(Type=Boot, Filename="%GAMEDIR%DerivedDataCache/Boot.ddc", MaxCacheSize=512)
 Local=(Type=FileSystem, ReadOnly=false, Clean=false, Flush=false, PurgeTransient=true, DeleteUnused=true, UnusedFileAge=34, FoldersToClean=-1, Path=../../../Engine/DerivedDataCache)
-Shared=(Type=FileSystem, ReadOnly=false, Clean=false, Flush=false, DeleteUnused=true, UnusedFileAge=23, FoldersToClean=-1, Path=\\epicgames.net\root\UE4DDC, EnvPathOverride=UE-SharedDataCachePath)
-AltShared=(Type=FileSystem, ReadOnly=true, Clean=false, Flush=false, DeleteUnused=true, UnusedFileAge=23, FoldersToClean=-1, Path=\\epicgames.net\root\UE4DDC2, EnvPathOverride=UE-SharedDataCachePath2)
+Shared=(Type=FileSystem, ReadOnly=false, Clean=false, Flush=false, DeleteUnused=false, UnusedFileAge=23, FoldersToClean=-1, Path=\\epicgames.net\root\UE4DDC, EnvPathOverride=UE-SharedDataCachePath)
+AltShared=(Type=FileSystem, ReadOnly=true, Clean=false, Flush=false, DeleteUnused=false, UnusedFileAge=23, FoldersToClean=-1, Path=\\epicgames.net\root\UE4DDC2, EnvPathOverride=UE-SharedDataCachePath2)
 PakRead=(Type=ReadPak, Filename="%GAMEDIR%DerivedDataCache/DDC.ddp")
 PakWrite=(Type=WritePak, Filename="%GAMEDIR%DerivedDataCache/DDC.ddp")
 
@@ -1983,8 +1983,8 @@ AsyncPut=(Type=AsyncPut, Inner=Hierarchy)
 Hierarchy=(Type=Hierarchical, Inner=Boot, Inner=EnginePak, Inner=PakWrite, Inner=PakRead, Inner=Local, Inner=Shared)
 Boot=(Type=Boot, Filename="%GAMEDIR%DerivedDataCache/Boot.ddc", MaxCacheSize=512)
 Local=(Type=FileSystem, ReadOnly=false, Clean=false, Flush=false, PurgeTransient=true, DeleteUnused=true, UnusedFileAge=34, FoldersToClean=-1, Path=../../../Engine/DerivedDataCache, EnvPathOverride=UE-LocalDataCachePath)
-Shared=(Type=FileSystem, ReadOnly=false, Clean=false, Flush=false, DeleteUnused=true, UnusedFileAge=23, FoldersToClean=-1, Path=\\epicgames.net\root\UE4DDC, EnvPathOverride=UE-SharedDataCachePath)
-AltShared=(Type=FileSystem, ReadOnly=true, Clean=false, Flush=false, DeleteUnused=true, UnusedFileAge=23, FoldersToClean=-1, Path=\\epicgames.net\root\UE4DDC2, EnvPathOverride=UE-SharedDataCachePath2)
+Shared=(Type=FileSystem, ReadOnly=false, Clean=false, Flush=false, DeleteUnused=false, UnusedFileAge=23, FoldersToClean=-1, Path=\\epicgames.net\root\UE4DDC, EnvPathOverride=UE-SharedDataCachePath)
+AltShared=(Type=FileSystem, ReadOnly=true, Clean=false, Flush=false, DeleteUnused=false, UnusedFileAge=23, FoldersToClean=-1, Path=\\epicgames.net\root\UE4DDC2, EnvPathOverride=UE-SharedDataCachePath2)
 EnginePak=(Type=ReadPak, Filename=../../../Engine/DerivedDataCache/DDC.ddp)
 PakRead=(Type=ReadPak, Filename="%GAMEDIR%DerivedDataCache/DDC.ddp")
 PakWrite=(Type=WritePak, Filename="%GAMEDIR%DerivedDataCache/DDC.ddp")
@@ -1996,8 +1996,8 @@ AsyncPut=(Type=AsyncPut, Inner=Hierarchy)
 Hierarchy=(Type=Hierarchical, Inner=Boot, Inner=PakWrite, Inner=PakRead, Inner=Local, Inner=Shared)
 Boot=(Type=Boot, Filename="%GAMEDIR%DerivedDataCache/Boot.ddc", MaxCacheSize=512)
 Local=(Type=FileSystem, ReadOnly=false, Clean=false, Flush=false, PurgeTransient=true, DeleteUnused=true, UnusedFileAge=34, FoldersToClean=-1, Path=../../../Engine/DerivedDataCache, EnvPathOverride=UE-LocalDataCachePath)
-Shared=(Type=FileSystem, ReadOnly=false, Clean=false, Flush=false, DeleteUnused=true, UnusedFileAge=23, FoldersToClean=-1, Path=\\epicgames.net\root\UE4DDC, EnvPathOverride=UE-SharedDataCachePath)
-AltShared=(Type=FileSystem, ReadOnly=true, Clean=false, Flush=false, DeleteUnused=true, UnusedFileAge=23, FoldersToClean=-1, Path=\\epicgames.net\root\UE4DDC2, EnvPathOverride=UE-SharedDataCachePath2)
+Shared=(Type=FileSystem, ReadOnly=false, Clean=false, Flush=false, DeleteUnused=false, UnusedFileAge=23, FoldersToClean=-1, Path=\\epicgames.net\root\UE4DDC, EnvPathOverride=UE-SharedDataCachePath)
+AltShared=(Type=FileSystem, ReadOnly=true, Clean=false, Flush=false, DeleteUnused=false, UnusedFileAge=23, FoldersToClean=-1, Path=\\epicgames.net\root\UE4DDC2, EnvPathOverride=UE-SharedDataCachePath2)
 PakRead=(Type=ReadPak, Filename=../../../Engine/DerivedDataCache/DDC.ddp)
 PakWrite=(Type=WritePak, Filename=../../../Engine/DerivedDataCache/DDC.ddp)
 
diff --git a/Engine/Source/Developer/DerivedDataCache/Private/DDCCleanup.cpp b/Engine/Source/Developer/DerivedDataCache/Private/DDCCleanup.cpp
--- a/Engine/Source/Developer/DerivedDataCache/Private/DDCCleanup.cpp
+++ b/Engine/Source/Developer/DerivedDataCache/Private/DDCCleanup.cpp
@@ -158,10 +158,10 @@ bool FDDCCleanup::CleanupFilesystemDirectory( TSharedPtr< FFilesystemInfo > File
 				}
 			}
 
-			if( ++NumFilesChecked > FilesystemInfo->MaxContinuousFileChecks && FilesystemInfo->MaxContinuousFileChecks > 0 && ShouldStop() == false )
+			if( ++NumFilesChecked >= FilesystemInfo->MaxContinuousFileChecks && FilesystemInfo->MaxContinuousFileChecks > 0 && ShouldStop() == false )
 			{
 				NumFilesChecked = 0;
-				Wait( 2.0f );
+				Wait( 1.0f );
 			}
 			else
 			{
diff --git a/Engine/Source/Developer/DerivedDataCache/Private/DerivedDataBackends.cpp b/Engine/Source/Developer/DerivedDataCache/Private/DerivedDataBackends.cpp
--- a/Engine/Source/Developer/DerivedDataCache/Private/DerivedDataBackends.cpp
+++ b/Engine/Source/Developer/DerivedDataCache/Private/DerivedDataBackends.cpp
@@ -513,6 +513,8 @@ public:
 			FParse::Value( Entry, TEXT("UnusedFileAge="), UnusedFileAge );
 			int32 MaxFoldersToClean = -1;
 			FParse::Value( Entry, TEXT("FoldersToClean="), MaxFoldersToClean );
+			int32 MaxFileChecksPerSec = -1;
+			FParse::Value( Entry, TEXT("MaxFileChecksPerSec="), MaxFileChecksPerSec );
 
 			if( bFlush )
 			{
@@ -528,7 +530,7 @@ public:
 			// Don't create the file system if shared data cache directory is not mounted
 			if( FCString::Strcmp(NodeName, TEXT("Shared")) != 0 || IFileManager::Get().DirectoryExists(*Path) )
 			{
-				InnerFileSystem = CreateFileSystemDerivedDataBackend( *Path, bReadOnly, bTouch, bPurgeTransient, bDeleteUnused, UnusedFileAge, MaxFoldersToClean );
+				InnerFileSystem = CreateFileSystemDerivedDataBackend( *Path, bReadOnly, bTouch, bPurgeTransient, bDeleteUnused, UnusedFileAge, MaxFoldersToClean, MaxFileChecksPerSec);
 			}
 
 			if( InnerFileSystem )
