Ensure that enum references in blueprint nodes to EMovementMode get updated correctly

diff --git a/Engine/Config/BaseEngine.ini b/Engine/Config/BaseEngine.ini
--- a/Engine/Config/BaseEngine.ini
+++ b/Engine/Config/BaseEngine.ini
@@ -220,6 +220,8 @@ bUseStreamingPause=false
 +ActiveClassRedirects=(OldClassName="PhysicsHingeActor",NewClassName="PhysicsConstraintActor")
 +ActiveClassRedirects=(OldClassName="PhysicsPrismaticActor",NewClassName="PhysicsConstraintActor")
 
++ActiveClassRedirects=(OldClassName="EMovementMode",NewClassName="/Script/Engine.EngineTypes:EMovementMode")
+
 +ActiveClassRedirects=(OldClassName="SensingComponent",NewClassName="PawnSensingComponent")
 +ActiveClassRedirects=(OldClassName="MovementComp_Character",NewClassName="CharacterMovementComponent")
 +ActiveClassRedirects=(OldClassName="MovementComp_Rotating",NewClassName="RotatingMovementComponent")
