Merged CL 2108710: nicer autocompletion on common commands

diff --git a/Engine/Config/BaseInput.ini b/Engine/Config/BaseInput.ini
--- a/Engine/Config/BaseInput.ini
+++ b/Engine/Config/BaseInput.ini
@@ -52,6 +52,8 @@ MaxScrollbackSize=1024
 +AutoCompleteMapPaths=Content/Maps
 +ManualAutoCompleteList=(Command="Exit",Desc="Exit (Exits the game)")
 +ManualAutoCompleteList=(Command="DebugCreatePlayer 1",Desc=)
++ManualAutoCompleteList=(Command="StartFPSChart",Desc="StartFPSChart (after that use EndFPSChart)")
++ManualAutoCompleteList=(Command="StopFPSChart",Desc="StopFPSChart (after that look for the output in Saved/Profiling/FPSChartStats)")
 +ManualAutoCompleteList=(Command="FreezeAt",Desc="Locks the player view and rendering time.")
 +ManualAutoCompleteList=(Command="Open",Desc="Open <MapName> (Opens the specified map, doesn't pass previously set options)")
 +ManualAutoCompleteList=(Command="Travel",Desc="Travel <MapName> (Travels to the specified map, passes along previously set options)")
