Fixed AIModule's class redirects not using full paths #UE4

diff --git a/Engine/Config/BaseEngine.ini b/Engine/Config/BaseEngine.ini
--- a/Engine/Config/BaseEngine.ini
+++ b/Engine/Config/BaseEngine.ini
@@ -469,101 +469,101 @@ bUseStreamingPause=false
 +ActiveClassRedirects=(OldClassName="AnimationCompressionAlgorithm_RevertToRaw",NewClassName="AnimCompress_RevertToRaw")
 
 ;AIModule redirections
-+ActiveClassRedirects=(OldClassName="AIController",NewClassName="/Script/AIModule.AIController")
-+ActiveClassRedirects=(OldClassName="AIResourceInterface",NewClassName="/Script/AIModule.AIResourceInterface")
-+ActiveClassRedirects=(OldClassName="AISystem",NewClassName="/Script/AIModule.AISystem")
-+ActiveClassRedirects=(OldClassName="AITypes",NewClassName="/Script/AIModule.AITypes")
-+ActiveClassRedirects=(OldClassName="BrainComponent",NewClassName="/Script/AIModule.BrainComponent")
-+ActiveClassRedirects=(OldClassName="KismetAIAsyncTaskProxy",NewClassName="/Script/AIModule.AIAsyncTaskBlueprintProxy")
-+ActiveClassRedirects=(OldClassName="KismetAIHelperLibrary",NewClassName="/Script/AIModule.AIBlueprintHelperLibrary")
-+ActiveClassRedirects=(OldClassName="BehaviorTree",NewClassName="/Script/AIModule.BehaviorTree")
-+ActiveClassRedirects=(OldClassName="BehaviorTreeComponent",NewClassName="/Script/AIModule.BehaviorTreeComponent")
-+ActiveClassRedirects=(OldClassName="BehaviorTreeManager",NewClassName="/Script/AIModule.BehaviorTreeManager")
-+ActiveClassRedirects=(OldClassName="BehaviorTreeTypes",NewClassName="/Script/AIModule.BehaviorTreeTypes")
-+ActiveClassRedirects=(OldClassName="BlackboardComponent",NewClassName="/Script/AIModule.BlackboardComponent")
-+ActiveClassRedirects=(OldClassName="BlackboardData",NewClassName="/Script/AIModule.BlackboardData")
-+ActiveClassRedirects=(OldClassName="BTAuxiliaryNode",NewClassName="/Script/AIModule.BTAuxiliaryNode")
-+ActiveClassRedirects=(OldClassName="BTCompositeNode",NewClassName="/Script/AIModule.BTCompositeNode")
-+ActiveClassRedirects=(OldClassName="BTDecorator",NewClassName="/Script/AIModule.BTDecorator")
-+ActiveClassRedirects=(OldClassName="BTFunctionLibrary",NewClassName="/Script/AIModule.BTFunctionLibrary")
-+ActiveClassRedirects=(OldClassName="BTNode",NewClassName="/Script/AIModule.BTNode")
-+ActiveClassRedirects=(OldClassName="BTService",NewClassName="/Script/AIModule.BTService")
-+ActiveClassRedirects=(OldClassName="BTTaskNode",NewClassName="/Script/AIModule.BTTaskNode")
-+ActiveClassRedirects=(OldClassName="BlackboardKeyAllTypes",NewClassName="/Script/AIModule.BlackboardKeyAllTypes")
-+ActiveClassRedirects=(OldClassName="BlackboardKeyType",NewClassName="/Script/AIModule.BlackboardKeyType")
-+ActiveClassRedirects=(OldClassName="BlackboardKeyType_Bool",NewClassName="/Script/AIModule.BlackboardKeyType_Bool")
-+ActiveClassRedirects=(OldClassName="BlackboardKeyType_Class",NewClassName="/Script/AIModule.BlackboardKeyType_Class")
-+ActiveClassRedirects=(OldClassName="BlackboardKeyType_Enum",NewClassName="/Script/AIModule.BlackboardKeyType_Enum")
-+ActiveClassRedirects=(OldClassName="BlackboardKeyType_Float",NewClassName="/Script/AIModule.BlackboardKeyType_Float")
-+ActiveClassRedirects=(OldClassName="BlackboardKeyType_Int",NewClassName="/Script/AIModule.BlackboardKeyType_Int")
-+ActiveClassRedirects=(OldClassName="BlackboardKeyType_Name",NewClassName="/Script/AIModule.BlackboardKeyType_Name")
-+ActiveClassRedirects=(OldClassName="BlackboardKeyType_NativeEnum",NewClassName="/Script/AIModule.BlackboardKeyType_NativeEnum")
-+ActiveClassRedirects=(OldClassName="BlackboardKeyType_Object",NewClassName="/Script/AIModule.BlackboardKeyType_Object")
-+ActiveClassRedirects=(OldClassName="BlackboardKeyType_String",NewClassName="/Script/AIModule.BlackboardKeyType_String")
-+ActiveClassRedirects=(OldClassName="BlackboardKeyType_Vector",NewClassName="/Script/AIModule.BlackboardKeyType_Vector")
-+ActiveClassRedirects=(OldClassName="BTComposite_Selector",NewClassName="/Script/AIModule.BTComposite_Selector")
-+ActiveClassRedirects=(OldClassName="BTComposite_Sequence",NewClassName="/Script/AIModule.BTComposite_Sequence")
-+ActiveClassRedirects=(OldClassName="BTComposite_SimpleParallel",NewClassName="/Script/AIModule.BTComposite_SimpleParallel")
-+ActiveClassRedirects=(OldClassName="BTDecorator_Blackboard",NewClassName="/Script/AIModule.BTDecorator_Blackboard")
-+ActiveClassRedirects=(OldClassName="BTDecorator_BlackboardBase",NewClassName="/Script/AIModule.BTDecorator_BlackboardBase")
-+ActiveClassRedirects=(OldClassName="BTDecorator_BlueprintBase",NewClassName="/Script/AIModule.BTDecorator_BlueprintBase")
-+ActiveClassRedirects=(OldClassName="BTDecorator_CompareBBEntries",NewClassName="/Script/AIModule.BTDecorator_CompareBBEntries")
-+ActiveClassRedirects=(OldClassName="BTDecorator_ConeCheck",NewClassName="/Script/AIModule.BTDecorator_ConeCheck")
-+ActiveClassRedirects=(OldClassName="BTDecorator_Cooldown",NewClassName="/Script/AIModule.BTDecorator_Cooldown")
-+ActiveClassRedirects=(OldClassName="BTDecorator_DoesPathExist",NewClassName="/Script/AIModule.BTDecorator_DoesPathExist")
-+ActiveClassRedirects=(OldClassName="BTDecorator_ForceSuccess",NewClassName="/Script/AIModule.BTDecorator_ForceSuccess")
-+ActiveClassRedirects=(OldClassName="BTDecorator_KeepInCone",NewClassName="/Script/AIModule.BTDecorator_KeepInCone")
-+ActiveClassRedirects=(OldClassName="BTDecorator_Loop",NewClassName="/Script/AIModule.BTDecorator_Loop")
-+ActiveClassRedirects=(OldClassName="BTDecorator_Optional",NewClassName="/Script/AIModule.BTDecorator_ForceSuccess")
-+ActiveClassRedirects=(OldClassName="BTDecorator_ReachedMoveGoal",NewClassName="/Script/AIModule.BTDecorator_ReachedMoveGoal")
-+ActiveClassRedirects=(OldClassName="BTDecorator_TimeLimit",NewClassName="/Script/AIModule.BTDecorator_TimeLimit")
-+ActiveClassRedirects=(OldClassName="BTService_BlackboardBase",NewClassName="/Script/AIModule.BTService_BlackboardBase")
-+ActiveClassRedirects=(OldClassName="BTService_BlueprintBase",NewClassName="/Script/AIModule.BTService_BlueprintBase")
-+ActiveClassRedirects=(OldClassName="BTService_DefaultFocus",NewClassName="/Script/AIModule.BTService_DefaultFocus")
-+ActiveClassRedirects=(OldClassName="BTTask_BlackboardBase",NewClassName="/Script/AIModule.BTTask_BlackboardBase")
-+ActiveClassRedirects=(OldClassName="BTTask_BlueprintBase",NewClassName="/Script/AIModule.BTTask_BlueprintBase")
-+ActiveClassRedirects=(OldClassName="BTTask_MakeNoise",NewClassName="/Script/AIModule.BTTask_MakeNoise")
-+ActiveClassRedirects=(OldClassName="BTTask_MoveDirectlyToward",NewClassName="/Script/AIModule.BTTask_MoveDirectlyToward")
-+ActiveClassRedirects=(OldClassName="BTTask_MoveTo",NewClassName="/Script/AIModule.BTTask_MoveTo")
-+ActiveClassRedirects=(OldClassName="BTTask_PlaySound",NewClassName="/Script/AIModule.BTTask_PlaySound")
-+ActiveClassRedirects=(OldClassName="BTTask_RunBehavior",NewClassName="/Script/AIModule.BTTask_RunBehavior")
-+ActiveClassRedirects=(OldClassName="BTTask_RunEQSQuery",NewClassName="/Script/AIModule.BTTask_RunEQSQuery")
-+ActiveClassRedirects=(OldClassName="BTTask_Wait",NewClassName="/Script/AIModule.BTTask_Wait")
-+ActiveClassRedirects=(OldClassName="EnvQuery",NewClassName="/Script/AIModule.EnvQuery")
-+ActiveClassRedirects=(OldClassName="EnvQueryContext",NewClassName="/Script/AIModule.EnvQueryContext")
-+ActiveClassRedirects=(OldClassName="EnvQueryGenerator",NewClassName="/Script/AIModule.EnvQueryGenerator")
-+ActiveClassRedirects=(OldClassName="EnvQueryManager",NewClassName="/Script/AIModule.EnvQueryManager")
-+ActiveClassRedirects=(OldClassName="EnvQueryOption",NewClassName="/Script/AIModule.EnvQueryOption")
-+ActiveClassRedirects=(OldClassName="EnvQueryTest",NewClassName="/Script/AIModule.EnvQueryTest")
-+ActiveClassRedirects=(OldClassName="EnvQueryTypes",NewClassName="/Script/AIModule.EnvQueryTypes")
-+ActiveClassRedirects=(OldClassName="EQSQueryResultSourceInterface",NewClassName="/Script/AIModule.EQSQueryResultSourceInterface")
-+ActiveClassRedirects=(OldClassName="EQSRenderingComponent",NewClassName="/Script/AIModule.EQSRenderingComponent")
-+ActiveClassRedirects=(OldClassName="EQSTestingPawn",NewClassName="/Script/AIModule.EQSTestingPawn")
-+ActiveClassRedirects=(OldClassName="EnvQueryContext_BlueprintBase",NewClassName="/Script/AIModule.EnvQueryContext_BlueprintBase")
-+ActiveClassRedirects=(OldClassName="EnvQueryContext_Item",NewClassName="/Script/AIModule.EnvQueryContext_Item")
-+ActiveClassRedirects=(OldClassName="EnvQueryContext_Querier",NewClassName="/Script/AIModule.EnvQueryContext_Querier")
-+ActiveClassRedirects=(OldClassName="EnvQueryGenerator_Composite",NewClassName="/Script/AIModule.EnvQueryGenerator_Composite")
-+ActiveClassRedirects=(OldClassName="EnvQueryGenerator_OnCircle",NewClassName="/Script/AIModule.EnvQueryGenerator_OnCircle")
-+ActiveClassRedirects=(OldClassName="EnvQueryGenerator_PathingGrid",NewClassName="/Script/AIModule.EnvQueryGenerator_PathingGrid")
-+ActiveClassRedirects=(OldClassName="EnvQueryGenerator_ProjectedPoints",NewClassName="/Script/AIModule.EnvQueryGenerator_ProjectedPoints")
-+ActiveClassRedirects=(OldClassName="EnvQueryGenerator_SimpleGrid",NewClassName="/Script/AIModule.EnvQueryGenerator_SimpleGrid")
-+ActiveClassRedirects=(OldClassName="EnvQueryAllItemTypes",NewClassName="/Script/AIModule.EnvQueryAllItemTypes")
-+ActiveClassRedirects=(OldClassName="EnvQueryItemType",NewClassName="/Script/AIModule.EnvQueryItemType")
-+ActiveClassRedirects=(OldClassName="EnvQueryItemType_Actor",NewClassName="/Script/AIModule.EnvQueryItemType_Actor")
-+ActiveClassRedirects=(OldClassName="EnvQueryItemType_ActorBase",NewClassName="/Script/AIModule.EnvQueryItemType_ActorBase")
-+ActiveClassRedirects=(OldClassName="EnvQueryItemType_Direction",NewClassName="/Script/AIModule.EnvQueryItemType_Direction")
-+ActiveClassRedirects=(OldClassName="EnvQueryItemType_Point",NewClassName="/Script/AIModule.EnvQueryItemType_Point")
-+ActiveClassRedirects=(OldClassName="EnvQueryItemType_VectorBase",NewClassName="/Script/AIModule.EnvQueryItemType_VectorBase")
-+ActiveClassRedirects=(OldClassName="EnvQueryTest_Distance",NewClassName="/Script/AIModule.EnvQueryTest_Distance")
-+ActiveClassRedirects=(OldClassName="EnvQueryTest_Dot",NewClassName="/Script/AIModule.EnvQueryTest_Dot")
-+ActiveClassRedirects=(OldClassName="EnvQueryTest_Pathfinding",NewClassName="/Script/AIModule.EnvQueryTest_Pathfinding")
-+ActiveClassRedirects=(OldClassName="EnvQueryTest_Trace",NewClassName="/Script/AIModule.EnvQueryTest_Trace")
-+ActiveClassRedirects=(OldClassName="CrowdAgentInterface",NewClassName="/Script/AIModule.CrowdAgentInterface")
-+ActiveClassRedirects=(OldClassName="CrowdFollowingComponent",NewClassName="/Script/AIModule.CrowdFollowingComponent")
-+ActiveClassRedirects=(OldClassName="CrowdManager",NewClassName="/Script/AIModule.CrowdManager")
-+ActiveClassRedirects=(OldClassName="NavigationComponent",NewClassName="/Script/AIModule.NavigationComponent")
-+ActiveClassRedirects=(OldClassName="PathFollowingComponent",NewClassName="/Script/AIModule.PathFollowingComponent")
-+ActiveClassRedirects=(OldClassName="PawnSensingComponent",NewClassName="/Script/AIModule.PawnSensingComponent")
++ActiveClassRedirects=(OldClassName="/Script/Engine.AIController",NewClassName="/Script/AIModule.AIController")
++ActiveClassRedirects=(OldClassName="/Script/Engine.AIResourceInterface",NewClassName="/Script/AIModule.AIResourceInterface")
++ActiveClassRedirects=(OldClassName="/Script/Engine.AISystem",NewClassName="/Script/AIModule.AISystem")
++ActiveClassRedirects=(OldClassName="/Script/Engine.AITypes",NewClassName="/Script/AIModule.AITypes")
++ActiveClassRedirects=(OldClassName="/Script/Engine.BrainComponent",NewClassName="/Script/AIModule.BrainComponent")
++ActiveClassRedirects=(OldClassName="/Script/Engine.KismetAIAsyncTaskProxy",NewClassName="/Script/AIModule.AIAsyncTaskBlueprintProxy")
++ActiveClassRedirects=(OldClassName="/Script/Engine.KismetAIHelperLibrary",NewClassName="/Script/AIModule.AIBlueprintHelperLibrary")
++ActiveClassRedirects=(OldClassName="/Script/Engine.BehaviorTree",NewClassName="/Script/AIModule.BehaviorTree")
++ActiveClassRedirects=(OldClassName="/Script/Engine.BehaviorTreeComponent",NewClassName="/Script/AIModule.BehaviorTreeComponent")
++ActiveClassRedirects=(OldClassName="/Script/Engine.BehaviorTreeManager",NewClassName="/Script/AIModule.BehaviorTreeManager")
++ActiveClassRedirects=(OldClassName="/Script/Engine.BehaviorTreeTypes",NewClassName="/Script/AIModule.BehaviorTreeTypes")
++ActiveClassRedirects=(OldClassName="/Script/Engine.BlackboardComponent",NewClassName="/Script/AIModule.BlackboardComponent")
++ActiveClassRedirects=(OldClassName="/Script/Engine.BlackboardData",NewClassName="/Script/AIModule.BlackboardData")
++ActiveClassRedirects=(OldClassName="/Script/Engine.BTAuxiliaryNode",NewClassName="/Script/AIModule.BTAuxiliaryNode")
++ActiveClassRedirects=(OldClassName="/Script/Engine.BTCompositeNode",NewClassName="/Script/AIModule.BTCompositeNode")
++ActiveClassRedirects=(OldClassName="/Script/Engine.BTDecorator",NewClassName="/Script/AIModule.BTDecorator")
++ActiveClassRedirects=(OldClassName="/Script/Engine.BTFunctionLibrary",NewClassName="/Script/AIModule.BTFunctionLibrary")
++ActiveClassRedirects=(OldClassName="/Script/Engine.BTNode",NewClassName="/Script/AIModule.BTNode")
++ActiveClassRedirects=(OldClassName="/Script/Engine.BTService",NewClassName="/Script/AIModule.BTService")
++ActiveClassRedirects=(OldClassName="/Script/Engine.BTTaskNode",NewClassName="/Script/AIModule.BTTaskNode")
++ActiveClassRedirects=(OldClassName="/Script/Engine.BlackboardKeyAllTypes",NewClassName="/Script/AIModule.BlackboardKeyAllTypes")
++ActiveClassRedirects=(OldClassName="/Script/Engine.BlackboardKeyType",NewClassName="/Script/AIModule.BlackboardKeyType")
++ActiveClassRedirects=(OldClassName="/Script/Engine.BlackboardKeyType_Bool",NewClassName="/Script/AIModule.BlackboardKeyType_Bool")
++ActiveClassRedirects=(OldClassName="/Script/Engine.BlackboardKeyType_Class",NewClassName="/Script/AIModule.BlackboardKeyType_Class")
++ActiveClassRedirects=(OldClassName="/Script/Engine.BlackboardKeyType_Enum",NewClassName="/Script/AIModule.BlackboardKeyType_Enum")
++ActiveClassRedirects=(OldClassName="/Script/Engine.BlackboardKeyType_Float",NewClassName="/Script/AIModule.BlackboardKeyType_Float")
++ActiveClassRedirects=(OldClassName="/Script/Engine.BlackboardKeyType_Int",NewClassName="/Script/AIModule.BlackboardKeyType_Int")
++ActiveClassRedirects=(OldClassName="/Script/Engine.BlackboardKeyType_Name",NewClassName="/Script/AIModule.BlackboardKeyType_Name")
++ActiveClassRedirects=(OldClassName="/Script/Engine.BlackboardKeyType_NativeEnum",NewClassName="/Script/AIModule.BlackboardKeyType_NativeEnum")
++ActiveClassRedirects=(OldClassName="/Script/Engine.BlackboardKeyType_Object",NewClassName="/Script/AIModule.BlackboardKeyType_Object")
++ActiveClassRedirects=(OldClassName="/Script/Engine.BlackboardKeyType_String",NewClassName="/Script/AIModule.BlackboardKeyType_String")
++ActiveClassRedirects=(OldClassName="/Script/Engine.BlackboardKeyType_Vector",NewClassName="/Script/AIModule.BlackboardKeyType_Vector")
++ActiveClassRedirects=(OldClassName="/Script/Engine.BTComposite_Selector",NewClassName="/Script/AIModule.BTComposite_Selector")
++ActiveClassRedirects=(OldClassName="/Script/Engine.BTComposite_Sequence",NewClassName="/Script/AIModule.BTComposite_Sequence")
++ActiveClassRedirects=(OldClassName="/Script/Engine.BTComposite_SimpleParallel",NewClassName="/Script/AIModule.BTComposite_SimpleParallel")
++ActiveClassRedirects=(OldClassName="/Script/Engine.BTDecorator_Blackboard",NewClassName="/Script/AIModule.BTDecorator_Blackboard")
++ActiveClassRedirects=(OldClassName="/Script/Engine.BTDecorator_BlackboardBase",NewClassName="/Script/AIModule.BTDecorator_BlackboardBase")
++ActiveClassRedirects=(OldClassName="/Script/Engine.BTDecorator_BlueprintBase",NewClassName="/Script/AIModule.BTDecorator_BlueprintBase")
++ActiveClassRedirects=(OldClassName="/Script/Engine.BTDecorator_CompareBBEntries",NewClassName="/Script/AIModule.BTDecorator_CompareBBEntries")
++ActiveClassRedirects=(OldClassName="/Script/Engine.BTDecorator_ConeCheck",NewClassName="/Script/AIModule.BTDecorator_ConeCheck")
++ActiveClassRedirects=(OldClassName="/Script/Engine.BTDecorator_Cooldown",NewClassName="/Script/AIModule.BTDecorator_Cooldown")
++ActiveClassRedirects=(OldClassName="/Script/Engine.BTDecorator_DoesPathExist",NewClassName="/Script/AIModule.BTDecorator_DoesPathExist")
++ActiveClassRedirects=(OldClassName="/Script/Engine.BTDecorator_ForceSuccess",NewClassName="/Script/AIModule.BTDecorator_ForceSuccess")
++ActiveClassRedirects=(OldClassName="/Script/Engine.BTDecorator_KeepInCone",NewClassName="/Script/AIModule.BTDecorator_KeepInCone")
++ActiveClassRedirects=(OldClassName="/Script/Engine.BTDecorator_Loop",NewClassName="/Script/AIModule.BTDecorator_Loop")
++ActiveClassRedirects=(OldClassName="/Script/Engine.BTDecorator_Optional",NewClassName="/Script/AIModule.BTDecorator_ForceSuccess")
++ActiveClassRedirects=(OldClassName="/Script/Engine.BTDecorator_ReachedMoveGoal",NewClassName="/Script/AIModule.BTDecorator_ReachedMoveGoal")
++ActiveClassRedirects=(OldClassName="/Script/Engine.BTDecorator_TimeLimit",NewClassName="/Script/AIModule.BTDecorator_TimeLimit")
++ActiveClassRedirects=(OldClassName="/Script/Engine.BTService_BlackboardBase",NewClassName="/Script/AIModule.BTService_BlackboardBase")
++ActiveClassRedirects=(OldClassName="/Script/Engine.BTService_BlueprintBase",NewClassName="/Script/AIModule.BTService_BlueprintBase")
++ActiveClassRedirects=(OldClassName="/Script/Engine.BTService_DefaultFocus",NewClassName="/Script/AIModule.BTService_DefaultFocus")
++ActiveClassRedirects=(OldClassName="/Script/Engine.BTTask_BlackboardBase",NewClassName="/Script/AIModule.BTTask_BlackboardBase")
++ActiveClassRedirects=(OldClassName="/Script/Engine.BTTask_BlueprintBase",NewClassName="/Script/AIModule.BTTask_BlueprintBase")
++ActiveClassRedirects=(OldClassName="/Script/Engine.BTTask_MakeNoise",NewClassName="/Script/AIModule.BTTask_MakeNoise")
++ActiveClassRedirects=(OldClassName="/Script/Engine.BTTask_MoveDirectlyToward",NewClassName="/Script/AIModule.BTTask_MoveDirectlyToward")
++ActiveClassRedirects=(OldClassName="/Script/Engine.BTTask_MoveTo",NewClassName="/Script/AIModule.BTTask_MoveTo")
++ActiveClassRedirects=(OldClassName="/Script/Engine.BTTask_PlaySound",NewClassName="/Script/AIModule.BTTask_PlaySound")
++ActiveClassRedirects=(OldClassName="/Script/Engine.BTTask_RunBehavior",NewClassName="/Script/AIModule.BTTask_RunBehavior")
++ActiveClassRedirects=(OldClassName="/Script/Engine.BTTask_RunEQSQuery",NewClassName="/Script/AIModule.BTTask_RunEQSQuery")
++ActiveClassRedirects=(OldClassName="/Script/Engine.BTTask_Wait",NewClassName="/Script/AIModule.BTTask_Wait")
++ActiveClassRedirects=(OldClassName="/Script/Engine.EnvQuery",NewClassName="/Script/AIModule.EnvQuery")
++ActiveClassRedirects=(OldClassName="/Script/Engine.EnvQueryContext",NewClassName="/Script/AIModule.EnvQueryContext")
++ActiveClassRedirects=(OldClassName="/Script/Engine.EnvQueryGenerator",NewClassName="/Script/AIModule.EnvQueryGenerator")
++ActiveClassRedirects=(OldClassName="/Script/Engine.EnvQueryManager",NewClassName="/Script/AIModule.EnvQueryManager")
++ActiveClassRedirects=(OldClassName="/Script/Engine.EnvQueryOption",NewClassName="/Script/AIModule.EnvQueryOption")
++ActiveClassRedirects=(OldClassName="/Script/Engine.EnvQueryTest",NewClassName="/Script/AIModule.EnvQueryTest")
++ActiveClassRedirects=(OldClassName="/Script/Engine.EnvQueryTypes",NewClassName="/Script/AIModule.EnvQueryTypes")
++ActiveClassRedirects=(OldClassName="/Script/Engine.EQSQueryResultSourceInterface",NewClassName="/Script/AIModule.EQSQueryResultSourceInterface")
++ActiveClassRedirects=(OldClassName="/Script/Engine.EQSRenderingComponent",NewClassName="/Script/AIModule.EQSRenderingComponent")
++ActiveClassRedirects=(OldClassName="/Script/Engine.EQSTestingPawn",NewClassName="/Script/AIModule.EQSTestingPawn")
++ActiveClassRedirects=(OldClassName="/Script/Engine.EnvQueryContext_BlueprintBase",NewClassName="/Script/AIModule.EnvQueryContext_BlueprintBase")
++ActiveClassRedirects=(OldClassName="/Script/Engine.EnvQueryContext_Item",NewClassName="/Script/AIModule.EnvQueryContext_Item")
++ActiveClassRedirects=(OldClassName="/Script/Engine.EnvQueryContext_Querier",NewClassName="/Script/AIModule.EnvQueryContext_Querier")
++ActiveClassRedirects=(OldClassName="/Script/Engine.EnvQueryGenerator_Composite",NewClassName="/Script/AIModule.EnvQueryGenerator_Composite")
++ActiveClassRedirects=(OldClassName="/Script/Engine.EnvQueryGenerator_OnCircle",NewClassName="/Script/AIModule.EnvQueryGenerator_OnCircle")
++ActiveClassRedirects=(OldClassName="/Script/Engine.EnvQueryGenerator_PathingGrid",NewClassName="/Script/AIModule.EnvQueryGenerator_PathingGrid")
++ActiveClassRedirects=(OldClassName="/Script/Engine.EnvQueryGenerator_ProjectedPoints",NewClassName="/Script/AIModule.EnvQueryGenerator_ProjectedPoints")
++ActiveClassRedirects=(OldClassName="/Script/Engine.EnvQueryGenerator_SimpleGrid",NewClassName="/Script/AIModule.EnvQueryGenerator_SimpleGrid")
++ActiveClassRedirects=(OldClassName="/Script/Engine.EnvQueryAllItemTypes",NewClassName="/Script/AIModule.EnvQueryAllItemTypes")
++ActiveClassRedirects=(OldClassName="/Script/Engine.EnvQueryItemType",NewClassName="/Script/AIModule.EnvQueryItemType")
++ActiveClassRedirects=(OldClassName="/Script/Engine.EnvQueryItemType_Actor",NewClassName="/Script/AIModule.EnvQueryItemType_Actor")
++ActiveClassRedirects=(OldClassName="/Script/Engine.EnvQueryItemType_ActorBase",NewClassName="/Script/AIModule.EnvQueryItemType_ActorBase")
++ActiveClassRedirects=(OldClassName="/Script/Engine.EnvQueryItemType_Direction",NewClassName="/Script/AIModule.EnvQueryItemType_Direction")
++ActiveClassRedirects=(OldClassName="/Script/Engine.EnvQueryItemType_Point",NewClassName="/Script/AIModule.EnvQueryItemType_Point")
++ActiveClassRedirects=(OldClassName="/Script/Engine.EnvQueryItemType_VectorBase",NewClassName="/Script/AIModule.EnvQueryItemType_VectorBase")
++ActiveClassRedirects=(OldClassName="/Script/Engine.EnvQueryTest_Distance",NewClassName="/Script/AIModule.EnvQueryTest_Distance")
++ActiveClassRedirects=(OldClassName="/Script/Engine.EnvQueryTest_Dot",NewClassName="/Script/AIModule.EnvQueryTest_Dot")
++ActiveClassRedirects=(OldClassName="/Script/Engine.EnvQueryTest_Pathfinding",NewClassName="/Script/AIModule.EnvQueryTest_Pathfinding")
++ActiveClassRedirects=(OldClassName="/Script/Engine.EnvQueryTest_Trace",NewClassName="/Script/AIModule.EnvQueryTest_Trace")
++ActiveClassRedirects=(OldClassName="/Script/Engine.CrowdAgentInterface",NewClassName="/Script/AIModule.CrowdAgentInterface")
++ActiveClassRedirects=(OldClassName="/Script/Engine.CrowdFollowingComponent",NewClassName="/Script/AIModule.CrowdFollowingComponent")
++ActiveClassRedirects=(OldClassName="/Script/Engine.CrowdManager",NewClassName="/Script/AIModule.CrowdManager")
++ActiveClassRedirects=(OldClassName="/Script/Engine.NavigationComponent",NewClassName="/Script/AIModule.NavigationComponent")
++ActiveClassRedirects=(OldClassName="/Script/Engine.PathFollowingComponent",NewClassName="/Script/AIModule.PathFollowingComponent")
++ActiveClassRedirects=(OldClassName="/Script/Engine.PawnSensingComponent",NewClassName="/Script/AIModule.PawnSensingComponent")
 
 ; Renamed SpriteComponent to BillboardComponent
 +ActiveClassRedirects=(OldClassName="SpriteComponent",NewClassName="BillboardComponent")
