Removing ipod touch 4 references, as it is not a supported target of ue4.

diff --git a/Engine/Config/BaseDeviceProfiles.ini b/Engine/Config/BaseDeviceProfiles.ini
--- a/Engine/Config/BaseDeviceProfiles.ini
+++ b/Engine/Config/BaseDeviceProfiles.ini
@@ -13,7 +13,6 @@
 +DeviceProfileNameAndTypes=iPhone4S,IOS
 +DeviceProfileNameAndTypes=iPhone5,IOS
 +DeviceProfileNameAndTypes=iPhone5S,IOS
-+DeviceProfileNameAndTypes=IPodTouch4,IOS
 +DeviceProfileNameAndTypes=iPodTouch5,IOS
 +DeviceProfileNameAndTypes=Android,Android
 +DeviceProfileNameAndTypes=PS4,PS4
@@ -99,11 +98,6 @@ BaseProfileName=IOS
 +CVars=r.DepthOfFieldQuality=1
 +CVars=r.LightShaftQuality=1
 
-[IPodTouch4 DeviceProfile]
-DeviceType=IOS
-BaseProfileName=IOS
-+CVars=r.RenderTargetSwitchWorkaround=1
-
 [iPodTouch5 DeviceProfile]
 DeviceType=IOS
 BaseProfileName=IOS
diff --git a/Engine/Source/Runtime/Core/Private/IOS/IOSPlatformMisc.cpp b/Engine/Source/Runtime/Core/Private/IOS/IOSPlatformMisc.cpp
--- a/Engine/Source/Runtime/Core/Private/IOS/IOSPlatformMisc.cpp
+++ b/Engine/Source/Runtime/Core/Private/IOS/IOSPlatformMisc.cpp
@@ -339,13 +339,7 @@ FIOSPlatformMisc::EIOSDevice FIOSPlatformMisc::GetIOSDeviceType()
 		// get major revision number
 		int Major = [DeviceIDString characterAtIndex:4] - '0';
 
-		// iPod4,1 is iPod Touch 4th gen
-		if (Major == 4)
-		{
-			DeviceType = IOS_IPodTouch4;
-		}
-		// iPod5 is iPod Touch 5th gen, anything higher will use 5th gen settings until released
-		else if (Major >= 5)
+		if (Major >= 5)
 		{
 			DeviceType = IOS_IPodTouch5;
 		}
diff --git a/Engine/Source/Runtime/Core/Public/IOS/IOSPlatformMisc.h b/Engine/Source/Runtime/Core/Public/IOS/IOSPlatformMisc.h
--- a/Engine/Source/Runtime/Core/Public/IOS/IOSPlatformMisc.h
+++ b/Engine/Source/Runtime/Core/Public/IOS/IOSPlatformMisc.h
@@ -81,7 +81,6 @@ struct CORE_API FIOSPlatformMisc : public FGenericPlatformMisc
 		IOS_IPhone4S,
 		IOS_IPhone5, // also the IPhone5C
 		IOS_IPhone5S,
-		IOS_IPodTouch4,
 		IOS_IPodTouch5,
 		IOS_IPad2,
 		IOS_IPad3,
@@ -102,7 +101,6 @@ struct CORE_API FIOSPlatformMisc : public FGenericPlatformMisc
 			L"IPhone4S",
 			L"IPhone5",
 			L"IPhone5S",
-			L"IPodTouch4",
 			L"IPodTouch5",
 			L"IPad2",
 			L"IPad3",
