Move the installed per-user DDC cache into a version-agnostic user folder, since it's designed to be shareable (and we're likely filling up the user's folder with stuff for every version).

diff --git a/Engine/Config/BaseEngine.ini b/Engine/Config/BaseEngine.ini
--- a/Engine/Config/BaseEngine.ini
+++ b/Engine/Config/BaseEngine.ini
@@ -1957,13 +1957,13 @@ Seattle=(Type=FileSystem, ReadOnly=false, Clean=false, Flush=false, DeleteUnused
 Pak=(Type=ReadPak, Filename="%GAMEDIR%DerivedDataCache/DDC.ddp")
 EnginePak=(Type=ReadPak, Filename=../../../Engine/DerivedDataCache/DDC.ddp)
 
-[RocketDerivedDataBackendGraph]
+[InstalledDerivedDataBackendGraph]
 MinimumDaysToKeepFile=7
 Root=(Type=KeyLength, Length=120, Inner=AsyncPut)
 AsyncPut=(Type=AsyncPut, Inner=Hierarchy)
 Hierarchy=(Type=Hierarchical, Inner=Boot, Inner=Pak, Inner=EnginePak, Inner=Local)
 Boot=(Type=Boot, Filename="%ENGINEUSERDIR%DerivedDataCache/Boot.ddc", MaxCacheSize=512)
-Local=(Type=FileSystem, ReadOnly=false, Clean=false, Flush=false, PurgeTransient=true, DeleteUnused=true, UnusedFileAge=34, FoldersToClean=-1, Path="%ENGINEUSERDIR%DerivedDataCache")
+Local=(Type=FileSystem, ReadOnly=false, Clean=false, Flush=false, PurgeTransient=true, DeleteUnused=true, UnusedFileAge=34, FoldersToClean=-1, Path="%ENGINEVERSIONAGNOSTICUSERDIR%DerivedDataCache")
 Pak=(Type=ReadPak, Filename="%GAMEDIR%DerivedDataCache/DDC.ddp")
 EnginePak=(Type=ReadPak, Filename=../../../Engine/DerivedDataCache/DDC.ddp)
 
diff --git a/Engine/Source/Developer/DerivedDataCache/Private/DerivedDataBackends.cpp b/Engine/Source/Developer/DerivedDataCache/Private/DerivedDataBackends.cpp
--- a/Engine/Source/Developer/DerivedDataCache/Private/DerivedDataBackends.cpp
+++ b/Engine/Source/Developer/DerivedDataCache/Private/DerivedDataBackends.cpp
@@ -72,7 +72,7 @@ public:
 		if( !RootCache )
 		{
 			// Use default graph
-			GraphName = FApp::IsEngineInstalled() ? TEXT("RocketDerivedDataBackendGraph") : TEXT("DerivedDataBackendGraph");
+			GraphName = FApp::IsEngineInstalled() ? TEXT("InstalledDerivedDataBackendGraph") : TEXT("DerivedDataBackendGraph");
 			FString Entry;
 			if( !GConfig->GetString( *GraphName, TEXT("Root"), Entry, GEngineIni ) || !Entry.Len())
 			{
diff --git a/Engine/Source/Runtime/Core/Private/Misc/ConfigCacheIni.cpp b/Engine/Source/Runtime/Core/Private/Misc/ConfigCacheIni.cpp
--- a/Engine/Source/Runtime/Core/Private/Misc/ConfigCacheIni.cpp
+++ b/Engine/Source/Runtime/Core/Private/Misc/ConfigCacheIni.cpp
@@ -179,6 +179,9 @@ void FConfigFile::CombineFromBuffer(const FString& Filename,const FString& Buffe
 	// Replace %ENGINEUSERDIR% with the user's engine directory.
 	Text = Text.Replace(TEXT("%ENGINEUSERDIR%"), *FPaths::EngineUserDir(), ESearchCase::CaseSensitive);
 
+	// Replace %ENGINEVERSIONAGNOSTICUSERDIR% with the user's engine directory.
+	Text = Text.Replace(TEXT("%ENGINEVERSIONAGNOSTICUSERDIR%"), *FPaths::EngineVersionAgnosticUserDir(), ESearchCase::CaseSensitive);
+
 	// Replace %APPSETTINGSDIR% with the game directory.
 	FString AppSettingsDir = FPlatformProcess::ApplicationSettingsDir();
 	FPaths::NormalizeFilename(AppSettingsDir);
diff --git a/Engine/Source/Runtime/Core/Private/Misc/Paths.cpp b/Engine/Source/Runtime/Core/Private/Misc/Paths.cpp
--- a/Engine/Source/Runtime/Core/Private/Misc/Paths.cpp
+++ b/Engine/Source/Runtime/Core/Private/Misc/Paths.cpp
@@ -35,6 +35,18 @@ FString FPaths::EngineUserDir()
 	}
 }
 
+FString FPaths::EngineVersionAgnosticUserDir()
+{
+	if (ShouldSaveToUserDir() || FApp::IsEngineInstalled())
+	{
+		return FPaths::Combine(FPlatformProcess::UserSettingsDir(), TEXT(EPIC_PRODUCT_IDENTIFIER), TEXT("Common")) + TEXT("/");
+	}
+	else
+	{
+		return FPaths::EngineDir();
+	}
+}
+
 FString FPaths::EngineContentDir()
 {
 	return FPaths::EngineDir() + TEXT("Content/");
diff --git a/Engine/Source/Runtime/Core/Public/Misc/Paths.h b/Engine/Source/Runtime/Core/Public/Misc/Paths.h
--- a/Engine/Source/Runtime/Core/Public/Misc/Paths.h
+++ b/Engine/Source/Runtime/Core/Public/Misc/Paths.h
@@ -28,6 +28,13 @@ public:
 	*/
 	static FString EngineUserDir();
 
+	/**
+	* Returns the root directory for user-specific engine files which can be shared between versions. Always writable.
+	*
+	* @return root user directory
+	*/
+	static FString EngineVersionAgnosticUserDir();
+
 	/** 
 	 * Returns the content directory of the "core" engine that can be shared across
 	 * several games or across games & mods. 
