Add a DDC backend (CreateProjectCache) to generate the DDC folder for a given project. Can be used to inject DDC into the user's cache for asset packs.

diff --git a/Engine/Config/BaseEngine.ini b/Engine/Config/BaseEngine.ini
--- a/Engine/Config/BaseEngine.ini
+++ b/Engine/Config/BaseEngine.ini
@@ -2013,6 +2013,16 @@ AltShared=(Type=FileSystem, ReadOnly=true, Clean=false, Flush=false, DeleteUnuse
 PakRead=(Type=ReadPak, Filename=../../../Engine/DerivedDataCache/DDC.ddp)
 PakWrite=(Type=WritePak, Filename=../../../Engine/DerivedDataCache/DDC.ddp)
 
+[CreateProjectCache]
+MinimumDaysToKeepFile=7
+Root=(Type=KeyLength, Length=120, Inner=AsyncPut)
+AsyncPut=(Type=AsyncPut, Inner=Hierarchy)
+Hierarchy=(Type=Hierarchical, Inner=Local, Inner=Project, Inner=Shared)
+Local=(Type=FileSystem, ReadOnly=false, Clean=false, Flush=false, PurgeTransient=true, DeleteUnused=true, UnusedFileAge=34, FoldersToClean=-1, Path=../../../Engine/DerivedDataCache, EnvPathOverride=UE-LocalDataCachePath)
+Project=(Type=FileSystem, ReadOnly=false, Clean=false, Flush=true, PurgeTransient=true, DeleteUnused=true, UnusedFileAge=34, FoldersToClean=-1, Path=%GAMEDIR%ProjectDerivedData)
+Shared=(Type=FileSystem, ReadOnly=false, Clean=false, Flush=false, DeleteUnused=false, UnusedFileAge=23, FoldersToClean=-1, Path=\\epicgames.net\root\UE4DDC, EnvPathOverride=UE-SharedDataCachePath)
+AltShared=(Type=FileSystem, ReadOnly=true, Clean=false, Flush=false, DeleteUnused=false, UnusedFileAge=23, FoldersToClean=-1, Path=\\epicgames.net\root\UE4DDC2, EnvPathOverride=UE-SharedDataCachePath2)
+
 [/Script/Engine.LocalPlayer]
 AspectRatioAxisConstraint=AspectRatio_MaintainXFOV
 
