Add r.RHICmdBypass to ConsoleVariables.ini

diff --git a/Engine/Config/ConsoleVariables.ini b/Engine/Config/ConsoleVariables.ini
--- a/Engine/Config/ConsoleVariables.ini
+++ b/Engine/Config/ConsoleVariables.ini
@@ -23,3 +23,5 @@
 ;r.ShaderDevelopmentMode=1
 ; Uncomment to dump shaders in the Saved folder
 ;r.DumpShaderDebugInfo=1
+; Uncomment to enable parallel rendering at startup
+;r.RHICmdBypass=0
\ No newline at end of file
