Integrate CL 2261242
"Buffered occlusion queries for ps4"

diff --git a/Engine/Config/ConsoleVariables.ini b/Engine/Config/ConsoleVariables.ini
--- a/Engine/Config/ConsoleVariables.ini
+++ b/Engine/Config/ConsoleVariables.ini
@@ -24,4 +24,6 @@
 ; Uncomment to dump shaders in the Saved folder
 ;r.DumpShaderDebugInfo=1
 ; Uncomment to enable parallel rendering at startup
-;r.RHICmdBypass=0
\ No newline at end of file
+;r.RHICmdBypass=0
+
+r.NumBufferedOcclusionQueries=3
\ No newline at end of file
