Back out revision 7 from //depot/UE4/Engine/Config/ConsoleVariables.ini

wrong branch

diff --git a/Engine/Config/ConsoleVariables.ini b/Engine/Config/ConsoleVariables.ini
--- a/Engine/Config/ConsoleVariables.ini
+++ b/Engine/Config/ConsoleVariables.ini
@@ -24,6 +24,4 @@
 ; Uncomment to dump shaders in the Saved folder
 ;r.DumpShaderDebugInfo=1
 ; Uncomment to enable parallel rendering at startup
-;r.RHICmdBypass=0
-
-r.NumBufferedOcclusionQueries=3
\ No newline at end of file
+;r.RHICmdBypass=0
\ No newline at end of file
