Dump shaders warning

diff --git a/Engine/Config/ConsoleVariables.ini b/Engine/Config/ConsoleVariables.ini
--- a/Engine/Config/ConsoleVariables.ini
+++ b/Engine/Config/ConsoleVariables.ini
@@ -22,6 +22,7 @@
 ; Uncomment to get detailed logs on shader compiles and the opportunity to retry on errors
 ;r.ShaderDevelopmentMode=1
 ; Uncomment to dump shaders in the Saved folder
+; Warning: leaving this on for a while will fill your hard drive with many small files and folders
 ;r.DumpShaderDebugInfo=1
 ; Uncomment to enable parallel rendering at startup
-;r.RHICmdBypass=0
\ No newline at end of file
+;r.RHICmdBypass=0
