Distance field lighting only enabled on Epic quality

diff --git a/Engine/Config/BaseScalability.ini b/Engine/Config/BaseScalability.ini
--- a/Engine/Config/BaseScalability.ini
+++ b/Engine/Config/BaseScalability.ini
@@ -61,6 +61,8 @@ r.Shadow.MaxResolution=512
 r.Shadow.RadiusThreshold=0.06
 r.Shadow.DistanceScale=0.6
 r.Shadow.CSM.TransitionScale=0
+r.DistanceFieldShadowing=0
+r.DistanceFieldAO=0
 
 [ShadowQuality@1]
 r.LightFunctionQuality=1
@@ -70,6 +72,8 @@ r.Shadow.MaxResolution=1024
 r.Shadow.RadiusThreshold=0.05
 r.Shadow.DistanceScale=0.7
 r.Shadow.CSM.TransitionScale=0.25
+r.DistanceFieldShadowing=0
+r.DistanceFieldAO=0
 
 [ShadowQuality@2]
 r.LightFunctionQuality=1
@@ -79,6 +83,8 @@ r.Shadow.MaxResolution=1024
 r.Shadow.RadiusThreshold=0.04
 r.Shadow.DistanceScale=0.85
 r.Shadow.CSM.TransitionScale=0.8
+r.DistanceFieldShadowing=0
+r.DistanceFieldAO=0
 
 [ShadowQuality@3]
 r.LightFunctionQuality=1
@@ -88,6 +94,8 @@ r.Shadow.MaxResolution=1024
 r.Shadow.RadiusThreshold=0.03
 r.Shadow.DistanceScale=1.0
 r.Shadow.CSM.TransitionScale=1.0
+r.DistanceFieldShadowing=1
+r.DistanceFieldAO=1
 
 ;-----------------------------------------------------------------------------------------------------------------
 
