#ue4 - Rename CameraComponent.bUseControllerViewRotation to bUsePawnViewRotation, and change its default to false.

- Also rename SpringArmComponent.bUseControllerViewRotation to bUsePawnViewRotation for consistency (default was already false).

- IMPORTANT UPGRADE NOTE: C++ projects using CameraComponents should ensure that they initialize bUsePawnViewRotation with the correct value if previously there was no initialization, since the default behavior has changed.
- The correct way to be backwards-compatible for CameraComponents is demonstrated in the C++ project templates:
   - add "_DEPRECATED" to any initialization of bUseControllerViewRotation (which becomes bUseControllerViewRotation_DEPRECATED).
   - add initialization of bUsePawnViewRotation to your desired default value.
   - if there was no previous initialization of bUsePawnViewRotation in code, then you MUST add initialization to "true" if you wish to maintain the same behavior as before (since the default has changed).

diff --git a/Engine/Config/BaseEngine.ini b/Engine/Config/BaseEngine.ini
--- a/Engine/Config/BaseEngine.ini
+++ b/Engine/Config/BaseEngine.ini
@@ -617,6 +617,8 @@ bUseStreamingPause=false
 +TaggedPropertyRedirects=(ClassName="ProjectileMovementComponent",OldPropertyName="Speed",NewPropertyName="InitialSpeed")
 +TaggedPropertyRedirects=(ClassName="PrimitiveComponent",OldPropertyName="CanBeCharacterBase",NewPropertyName="CanCharacterStepUpOn")
 
++TaggedPropertyRedirects=(ClassName="SpringArmComponent",OldPropertyName="bUseControllerViewRotation",NewPropertyName="bUsePawnViewRotation")
+
 +TaggedPropertyRedirects=(ClassName="MaterialEditorInstanceConstant",OldPropertyName="RefractionBias",NewPropertyName="RefractionDepthBias")
 +TaggedPropertyRedirects=(ClassName="Material",OldPropertyName="RefractionBias",NewPropertyName="RefractionDepthBias")
 +TaggedPropertyRedirects=(ClassName="SlateBrush",OldPropertyName="TextureObject",NewPropertyName="ResourceObject")
@@ -978,6 +980,9 @@ bUseStreamingPause=false
 +K2ParamRedirects=(NodeName="/Script/BlueprintGraph.K2Node_CallFunction", OldParamName="CharacterMovementComponent.AddImpulse.InMomentum", NewParamName="Impulse")
 +K2ParamRedirects=(NodeName="/Script/BlueprintGraph.K2Node_CallFunction", OldParamName="CharacterMovementComponent.AddImpulse.bMassIndependent", NewParamName="bVelocityChange")
 
++K2FieldRedirects=(OldFieldName="CameraComponent.bUseControllerViewRotation",NewPropertyName="CameraComponent.bUsePawnViewRotation")
++K2FieldRedirects=(OldFieldName="SpringArmComponent.bUseControllerViewRotation",NewPropertyName="SpringArmComponent.bUsePawnViewRotation")
+
 +K2FieldRedirects=(OldFieldName="MovementComponent.GetModifiedMaxSpeed",NewFieldName="MovementComponent.K2_GetModifiedMaxSpeed")
 +K2FieldRedirects=(OldFieldName="MovementComponent.GetMaxSpeedModifier",NewFieldName="MovementComponent.K2_GetMaxSpeedModifier")
 +K2FieldRedirects=(OldFieldName="CharacterMovementComponent.GetModifiedMaxAcceleration",NewFieldName="CharacterMovementComponent.K2_GetModifiedMaxAcceleration")
diff --git a/Engine/Source/Editor/DetailCustomizations/Private/CameraDetails.cpp b/Engine/Source/Editor/DetailCustomizations/Private/CameraDetails.cpp
--- a/Engine/Source/Editor/DetailCustomizations/Private/CameraDetails.cpp
+++ b/Engine/Source/Editor/DetailCustomizations/Private/CameraDetails.cpp
@@ -87,7 +87,7 @@ void FCameraDetails::CustomizeDetails( IDetailLayoutBuilder& DetailLayout )
 					]
 			];
 
-	CameraCategory.AddProperty(DetailLayout.GetProperty(GET_MEMBER_NAME_CHECKED(UCameraComponent, bUseControllerViewRotation)));
+	CameraCategory.AddProperty(DetailLayout.GetProperty(GET_MEMBER_NAME_CHECKED(UCameraComponent, bUsePawnViewRotation)));
 	CameraCategory.AddProperty(DetailLayout.GetProperty(GET_MEMBER_NAME_CHECKED(UCameraComponent, PostProcessBlendWeight)));
 
 	UpdateAspectTextFromProperty();
diff --git a/Engine/Source/Runtime/Core/Public/UObject/ObjectVersion.h b/Engine/Source/Runtime/Core/Public/UObject/ObjectVersion.h
--- a/Engine/Source/Runtime/Core/Public/UObject/ObjectVersion.h
+++ b/Engine/Source/Runtime/Core/Public/UObject/ObjectVersion.h
@@ -610,6 +610,9 @@ enum EUnrealEngineObjectUE4Version
 	VER_UE4_DEPRECATE_UMG_STYLE_ASSETS,
 	// Duplicating Blueprints will regenerate NodeGuids after this version
 	VER_UE4_POST_DUPLICATE_NODE_GUID,
+	// Rename USpringArmComponent::bUseControllerViewRotation to bUsePawnViewRotation,
+	// Rename UCameraComponent::bUseControllerViewRotation to bUsePawnViewRotation (and change the default value)
+	VER_UE4_RENAME_CAMERA_COMPONENT_VIEW_ROTATION,
 
 	// -----<new versions can be added before this line>-------------------------------------------------
 	// - this needs to be the last line (see note below)
diff --git a/Engine/Source/Runtime/Engine/Classes/Camera/CameraComponent.h b/Engine/Source/Runtime/Engine/Classes/Camera/CameraComponent.h
--- a/Engine/Source/Runtime/Engine/Classes/Camera/CameraComponent.h
+++ b/Engine/Source/Runtime/Engine/Classes/Camera/CameraComponent.h
@@ -32,7 +32,7 @@ class UCameraComponent : public USceneComponent
 
 	/** If this camera component is placed on a pawn, should it use the view rotation of the pawn where possible? */
 	UPROPERTY(EditAnywhere, BlueprintReadWrite, Category=CameraSettings)
-	uint32 bUseControllerViewRotation:1;
+	uint32 bUsePawnViewRotation:1;
 
 	// The type of camera
 	UPROPERTY(Interp, EditAnywhere, BlueprintReadWrite, Category=CameraSettings)
@@ -49,6 +49,7 @@ class UCameraComponent : public USceneComponent
 	// UActorComponent interface
 	ENGINE_API virtual void OnRegister() override;
 	ENGINE_API virtual void OnUnregister() override;
+	ENGINE_API virtual void PostLoad() override;
 #if WITH_EDITOR
 	ENGINE_API virtual void CheckForErrors() override;
 	// End of UActorComponent interface
@@ -87,4 +88,9 @@ public:
 	ENGINE_API void OverrideFrustumColor(FColor OverrideColor);
 	ENGINE_API void RestoreFrustumColor();
 #endif
+
+public:
+	/** DEPRECATED: use bUsePawnViewRotation instead */
+	UPROPERTY()
+	uint32 bUseControllerViewRotation_DEPRECATED:1;
 };
\ No newline at end of file
diff --git a/Engine/Source/Runtime/Engine/Classes/GameFramework/SpringArmComponent.h b/Engine/Source/Runtime/Engine/Classes/GameFramework/SpringArmComponent.h
--- a/Engine/Source/Runtime/Engine/Classes/GameFramework/SpringArmComponent.h
+++ b/Engine/Source/Runtime/Engine/Classes/GameFramework/SpringArmComponent.h
@@ -42,7 +42,7 @@ class ENGINE_API USpringArmComponent : public USceneComponent
 
 	/** If this component is placed on a pawn, should it use the view rotation of the pawn where possible? */
 	UPROPERTY(EditAnywhere, BlueprintReadWrite, Category=CameraSettings)
-	uint32 bUseControllerViewRotation:1;
+	uint32 bUsePawnViewRotation:1;
 
 	/** Should we inherit pitch from parent component. Does nothing if using Absolute Rotation. */
 	UPROPERTY(EditAnywhere, BlueprintReadWrite, Category=CameraSettings)
diff --git a/Engine/Source/Runtime/Engine/Private/Camera/CameraComponent.cpp b/Engine/Source/Runtime/Engine/Private/Camera/CameraComponent.cpp
--- a/Engine/Source/Runtime/Engine/Private/Camera/CameraComponent.cpp
+++ b/Engine/Source/Runtime/Engine/Private/Camera/CameraComponent.cpp
@@ -26,7 +26,8 @@ UCameraComponent::UCameraComponent(const class FPostConstructInitializePropertie
 	OrthoWidth = 512.0f;
 	bConstrainAspectRatio = false;
 	PostProcessBlendWeight = 1.0f;
-	bUseControllerViewRotation = true;
+	bUseControllerViewRotation_DEPRECATED = true; // the previous default value before bUsePawnViewRotation replaced this var.
+	bUsePawnViewRotation = false;
 	bAutoActivate = true;
 }
 
@@ -86,6 +87,20 @@ void UCameraComponent::OnUnregister()
 #endif
 }
 
+
+void UCameraComponent::PostLoad()
+{
+	Super::PostLoad();
+
+	const int32 LinkerUE4Ver = GetLinkerUE4Version();
+
+	if (LinkerUE4Ver < VER_UE4_RENAME_CAMERA_COMPONENT_VIEW_ROTATION)
+	{
+		bUsePawnViewRotation = bUseControllerViewRotation_DEPRECATED;
+	}
+}
+
+
 #if WITH_EDITOR
 void UCameraComponent::RefreshVisualRepresentation()
 {
@@ -129,7 +144,7 @@ void UCameraComponent::PostEditChangeProperty(FPropertyChangedEvent& PropertyCha
 
 void UCameraComponent::GetCameraView(float DeltaTime, FMinimalViewInfo& DesiredView)
 {
-	if (bUseControllerViewRotation)
+	if (bUsePawnViewRotation)
 	{
 		if (APawn* OwningPawn = Cast<APawn>(GetOwner()))
 		{
diff --git a/Engine/Source/Runtime/Engine/Private/GameFramework/SpringArmComponent.cpp b/Engine/Source/Runtime/Engine/Private/GameFramework/SpringArmComponent.cpp
--- a/Engine/Source/Runtime/Engine/Private/GameFramework/SpringArmComponent.cpp
+++ b/Engine/Source/Runtime/Engine/Private/GameFramework/SpringArmComponent.cpp
@@ -17,7 +17,7 @@ USpringArmComponent::USpringArmComponent(const class FPostConstructInitializePro
 	
 	bAutoActivate = true;
 	bTickInEditor = true;
-	bUseControllerViewRotation = false;
+	bUsePawnViewRotation = false;
 	bDoCollisionTest = true;
 
 	bInheritPitch = true;
@@ -123,7 +123,7 @@ void USpringArmComponent::TickComponent(float DeltaTime, enum ELevelTick TickTyp
 {
 	Super::TickComponent(DeltaTime, TickType, ThisTickFunction);
 
-	if (bUseControllerViewRotation)
+	if (bUsePawnViewRotation)
 	{
 		if (APawn* OwningPawn = Cast<APawn>(GetOwner()))
 		{
