#ue4 - Add blueprint ini redirect for CharacterMovementComponent.bCrouchMovesCharacterDown -> bCrouchMaintainsBaseLocation

diff --git a/Engine/Config/BaseEngine.ini b/Engine/Config/BaseEngine.ini
--- a/Engine/Config/BaseEngine.ini
+++ b/Engine/Config/BaseEngine.ini
@@ -976,6 +976,7 @@ bUseStreamingPause=false
 +K2FieldRedirects=(OldFieldName="CharacterMovementComponent.WaterSpeed",NewPropertyName="CharacterMovementComponent.MaxSwimSpeed")
 +K2FieldRedirects=(OldFieldName="CharacterMovementComponent.CrouchedSpeedPercent",NewPropertyName="CharacterMovementComponent.CrouchedSpeedMultiplier")
 +K2FieldRedirects=(OldFieldName="CharacterMovementComponent.bOrientToMovement",NewPropertyName="CharacterMovementComponent.bOrientRotationToMovement")
++K2FieldRedirects=(OldFieldName="CharacterMovementComponent.bCrouchMovesCharacterDown",NewPropertyName="CharacterMovementComponent.bCrouchMaintainsBaseLocation")
 +K2FieldRedirects=(OldFieldName="CharacterMovementComponent.AddMomentum",NewFieldName="CharacterMovementComponent.AddImpulse")
 +K2ParamRedirects=(NodeName="/Script/BlueprintGraph.K2Node_CallFunction", OldParamName="CharacterMovementComponent.AddImpulse.InMomentum", NewParamName="Impulse")
 +K2ParamRedirects=(NodeName="/Script/BlueprintGraph.K2Node_CallFunction", OldParamName="CharacterMovementComponent.AddImpulse.bMassIndependent", NewParamName="bVelocityChange")
