Renamed LightComponent's SetBrightness to SetIntensity.

diff --git a/Engine/Config/BaseEngine.ini b/Engine/Config/BaseEngine.ini
--- a/Engine/Config/BaseEngine.ini
+++ b/Engine/Config/BaseEngine.ini
@@ -999,6 +999,8 @@ bUseStreamingPause=false
 +K2FieldRedirects=(OldFieldName="ExponentialHeightFogComponent.SetEnabled",NewFieldName="SceneComponent.SetVisibility")
 +K2FieldRedirects=(OldFieldName="LightComponent.SetEnabled",NewFieldName="SceneComponent.SetVisibility")
 +K2FieldRedirects=(OldFieldName="LightComponent.ToggleEnabled",NewFieldName="SceneComponent.ToggleVisibility")
++K2FieldRedirects=(OldFieldName="LightComponent.SetBrightness",NewFieldName="LightComponent.SetIntensity")
++K2FieldRedirects=(OldFieldName="SkyLightComponent.SetBrightness",NewFieldName="SkyLightComponent.SetIntensity")
 +K2FieldRedirects=(OldFieldName="ParticleSystemComponent.SetActive",NewFieldName="SceneComponent.SetActive")
 +K2FieldRedirects=(OldFieldName="ParticleSystemComponent.ActivateSystem",NewFieldName="SceneComponent.Activate")
 +K2FieldRedirects=(OldFieldName="ParticleSystemComponent.DeactivateSystem",NewFieldName="SceneComponent.Deactivate")
diff --git a/Engine/Source/Runtime/Engine/Classes/Components/LightComponent.h b/Engine/Source/Runtime/Engine/Classes/Components/LightComponent.h
--- a/Engine/Source/Runtime/Engine/Classes/Components/LightComponent.h
+++ b/Engine/Source/Runtime/Engine/Classes/Components/LightComponent.h
@@ -150,9 +150,9 @@ class ENGINE_API ULightComponent : public ULightComponentBase
 	FColor BloomTint;
 
 public:
-	/** Set brightness of the light */
+	/** Set intensity of the light */
 	UFUNCTION(BlueprintCallable, Category="Rendering|Components|Light")
-	void SetBrightness(float NewBrightness);
+	void SetIntensity(float NewIntensity);
 
 	/** Set color of the light */
 	UFUNCTION(BlueprintCallable, Category="Rendering|Components|Light")
diff --git a/Engine/Source/Runtime/Engine/Classes/Components/SkyLightComponent.h b/Engine/Source/Runtime/Engine/Classes/Components/SkyLightComponent.h
--- a/Engine/Source/Runtime/Engine/Classes/Components/SkyLightComponent.h
+++ b/Engine/Source/Runtime/Engine/Classes/Components/SkyLightComponent.h
@@ -147,7 +147,7 @@ class ENGINE_API USkyLightComponent : public ULightComponentBase
 	void CaptureEmissiveIrradianceEnvironmentMap(FSHVectorRGB3& OutIrradianceMap) const;
 
 	UFUNCTION(BlueprintCallable, Category="Rendering|Components|SkyLight")
-	void SetBrightness(float NewBrightness);
+	void SetIntensity(float NewIntensity);
 
 	/** Set color of the light */
 	UFUNCTION(BlueprintCallable, Category="Rendering|Components|SkyLight")
diff --git a/Engine/Source/Runtime/Engine/Private/Light.cpp b/Engine/Source/Runtime/Engine/Private/Light.cpp
--- a/Engine/Source/Runtime/Engine/Private/Light.cpp
+++ b/Engine/Source/Runtime/Engine/Private/Light.cpp
@@ -103,7 +103,7 @@ void ALight::SetBrightness(float NewBrightness)
 {
 	if(LightComponent)
 	{
-		LightComponent->SetBrightness(NewBrightness);
+		LightComponent->SetIntensity(NewBrightness);
 	}
 }
 
diff --git a/Engine/Source/Runtime/Engine/Private/LightComponent.cpp b/Engine/Source/Runtime/Engine/Private/LightComponent.cpp
--- a/Engine/Source/Runtime/Engine/Private/LightComponent.cpp
+++ b/Engine/Source/Runtime/Engine/Private/LightComponent.cpp
@@ -644,13 +644,13 @@ void ULightComponent::DestroyRenderState_Concurrent()
 }
 
 /** Set brightness of the light */
-void ULightComponent::SetBrightness(float NewBrightness)
+void ULightComponent::SetIntensity(float NewIntensity)
 {
 	// Can't set brightness on a static light
 	if (!(IsRegistered() && Mobility == EComponentMobility::Static)
-		&& Intensity != NewBrightness)
+		&& Intensity != NewIntensity)
 	{
-		Intensity = NewBrightness;
+		Intensity = NewIntensity;
 
 		// Use lightweight color and brightness update 
 		if( World && World->Scene )
diff --git a/Engine/Source/Runtime/Engine/Private/SkyLightComponent.cpp b/Engine/Source/Runtime/Engine/Private/SkyLightComponent.cpp
--- a/Engine/Source/Runtime/Engine/Private/SkyLightComponent.cpp
+++ b/Engine/Source/Runtime/Engine/Private/SkyLightComponent.cpp
@@ -425,13 +425,13 @@ void USkyLightComponent::CaptureEmissiveIrradianceEnvironmentMap(FSHVectorRGB3&
 }
 
 /** Set brightness of the light */
-void USkyLightComponent::SetBrightness(float NewBrightness)
+void USkyLightComponent::SetIntensity(float NewIntensity)
 {
 	// Can't set brightness on a static light
 	if (!(IsRegistered() && Mobility == EComponentMobility::Static)
-		&& Intensity != NewBrightness)
+		&& Intensity != NewIntensity)
 	{
-		Intensity = NewBrightness;
+		Intensity = NewIntensity;
 		MarkRenderStateDirty();
 	}
 }
