Enabled import textures and import materials by default on FBX import.
#ttp 345827 - NUX: MVP: Turn on 'import textures/materials' by default

diff --git a/Engine/Config/BaseEditorUserSettings.ini b/Engine/Config/BaseEditorUserSettings.ini
--- a/Engine/Config/BaseEditorUserSettings.ini
+++ b/Engine/Config/BaseEditorUserSettings.ini
@@ -538,10 +538,10 @@ BackupIntervalInMinutes=5
 
 
 [/Script/UnrealEd.FbxImportUI]
-bImportMaterials=False
+bImportMaterials=True
 bAutoCreateGroups=True
 bInvertNormalMaps=False
-bImportTextures=False
+bImportTextures=True
 bOverrideFullName=True
 bConvertScene=True
 bRemoveNamespaces=True
