#ue4 - Rename USpringArmComponent::bUsePawnViewRotation to bUsePawnControlRotation. Rename UCameraComponent::bUsePawnViewRotation to bUsePawnControlRotation.

diff --git a/Engine/Config/BaseEngine.ini b/Engine/Config/BaseEngine.ini
--- a/Engine/Config/BaseEngine.ini
+++ b/Engine/Config/BaseEngine.ini
@@ -617,7 +617,9 @@ bUseStreamingPause=false
 +TaggedPropertyRedirects=(ClassName="ProjectileMovementComponent",OldPropertyName="Speed",NewPropertyName="InitialSpeed")
 +TaggedPropertyRedirects=(ClassName="PrimitiveComponent",OldPropertyName="CanBeCharacterBase",NewPropertyName="CanCharacterStepUpOn")
 
-+TaggedPropertyRedirects=(ClassName="SpringArmComponent",OldPropertyName="bUseControllerViewRotation",NewPropertyName="bUsePawnViewRotation")
++TaggedPropertyRedirects=(ClassName="SpringArmComponent",OldPropertyName="bUseControllerViewRotation",NewPropertyName="bUsePawnControlRotation")
++TaggedPropertyRedirects=(ClassName="SpringArmComponent",OldPropertyName="bUsePawnViewRotation",NewPropertyName="bUsePawnControlRotation")
++TaggedPropertyRedirects=(ClassName="CameraComponent",OldPropertyName="bUsePawnViewRotation",NewPropertyName="bUsePawnControlRotation")
 
 +TaggedPropertyRedirects=(ClassName="MaterialEditorInstanceConstant",OldPropertyName="RefractionBias",NewPropertyName="RefractionDepthBias")
 +TaggedPropertyRedirects=(ClassName="Material",OldPropertyName="RefractionBias",NewPropertyName="RefractionDepthBias")
@@ -981,8 +983,10 @@ bUseStreamingPause=false
 +K2ParamRedirects=(NodeName="/Script/BlueprintGraph.K2Node_CallFunction", OldParamName="CharacterMovementComponent.AddImpulse.InMomentum", NewParamName="Impulse")
 +K2ParamRedirects=(NodeName="/Script/BlueprintGraph.K2Node_CallFunction", OldParamName="CharacterMovementComponent.AddImpulse.bMassIndependent", NewParamName="bVelocityChange")
 
-+K2FieldRedirects=(OldFieldName="CameraComponent.bUseControllerViewRotation",NewPropertyName="CameraComponent.bUsePawnViewRotation")
-+K2FieldRedirects=(OldFieldName="SpringArmComponent.bUseControllerViewRotation",NewPropertyName="SpringArmComponent.bUsePawnViewRotation")
++K2FieldRedirects=(OldFieldName="CameraComponent.bUseControllerViewRotation",NewPropertyName="CameraComponent.bUsePawnControlRotation")
++K2FieldRedirects=(OldFieldName="CameraComponent.bUsePawnViewRotation",NewPropertyName="CameraComponent.bUsePawnControlRotation")
++K2FieldRedirects=(OldFieldName="SpringArmComponent.bUseControllerViewRotation",NewPropertyName="SpringArmComponent.bUsePawnControlRotation")
++K2FieldRedirects=(OldFieldName="SpringArmComponent.bUsePawnViewRotation",NewPropertyName="SpringArmComponent.bUsePawnControlRotation")
 
 +K2FieldRedirects=(OldFieldName="MovementComponent.GetModifiedMaxSpeed",NewFieldName="MovementComponent.K2_GetModifiedMaxSpeed")
 +K2FieldRedirects=(OldFieldName="MovementComponent.GetMaxSpeedModifier",NewFieldName="MovementComponent.K2_GetMaxSpeedModifier")
diff --git a/Engine/Source/Editor/DetailCustomizations/Private/CameraDetails.cpp b/Engine/Source/Editor/DetailCustomizations/Private/CameraDetails.cpp
--- a/Engine/Source/Editor/DetailCustomizations/Private/CameraDetails.cpp
+++ b/Engine/Source/Editor/DetailCustomizations/Private/CameraDetails.cpp
@@ -87,7 +87,7 @@ void FCameraDetails::CustomizeDetails( IDetailLayoutBuilder& DetailLayout )
 					]
 			];
 
-	CameraCategory.AddProperty(DetailLayout.GetProperty(GET_MEMBER_NAME_CHECKED(UCameraComponent, bUsePawnViewRotation)));
+	CameraCategory.AddProperty(DetailLayout.GetProperty(GET_MEMBER_NAME_CHECKED(UCameraComponent, bUsePawnControlRotation)));
 	CameraCategory.AddProperty(DetailLayout.GetProperty(GET_MEMBER_NAME_CHECKED(UCameraComponent, PostProcessBlendWeight)));
 
 	UpdateAspectTextFromProperty();
diff --git a/Engine/Source/Runtime/Core/Public/UObject/ObjectVersion.h b/Engine/Source/Runtime/Core/Public/UObject/ObjectVersion.h
--- a/Engine/Source/Runtime/Core/Public/UObject/ObjectVersion.h
+++ b/Engine/Source/Runtime/Core/Public/UObject/ObjectVersion.h
@@ -615,6 +615,9 @@ enum EUnrealEngineObjectUE4Version
 	VER_UE4_RENAME_CAMERA_COMPONENT_VIEW_ROTATION,
 	// Changed FName to be case preserving
 	VER_UE4_CASE_PRESERVING_FNAME,
+	// Rename USpringArmComponent::bUsePawnViewRotation to bUsePawnControlRotation
+	// Rename UCameraComponent::bUsePawnViewRotation to bUsePawnControlRotation
+	VER_UE4_RENAME_CAMERA_COMPONENT_CONTROL_ROTATION,
 
 	// -----<new versions can be added before this line>-------------------------------------------------
 	// - this needs to be the last line (see note below)
diff --git a/Engine/Source/Runtime/Engine/Classes/Camera/CameraComponent.h b/Engine/Source/Runtime/Engine/Classes/Camera/CameraComponent.h
--- a/Engine/Source/Runtime/Engine/Classes/Camera/CameraComponent.h
+++ b/Engine/Source/Runtime/Engine/Classes/Camera/CameraComponent.h
@@ -30,9 +30,12 @@ class UCameraComponent : public USceneComponent
 	UPROPERTY(Interp, EditAnywhere, BlueprintReadWrite, Category=CameraSettings)
 	uint32 bConstrainAspectRatio:1;
 
-	/** If this camera component is placed on a pawn, should it use the view rotation of the pawn where possible? */
+	/**
+	 * If this camera component is placed on a pawn, should it use the view/control rotation of the pawn where possible?
+	 * @see APawn::GetViewRotation()
+	 */
 	UPROPERTY(EditAnywhere, BlueprintReadWrite, Category=CameraSettings)
-	uint32 bUsePawnViewRotation:1;
+	uint32 bUsePawnControlRotation:1;
 
 	// The type of camera
 	UPROPERTY(Interp, EditAnywhere, BlueprintReadWrite, Category=CameraSettings)
@@ -90,7 +93,7 @@ public:
 #endif
 
 public:
-	/** DEPRECATED: use bUsePawnViewRotation instead */
+	/** DEPRECATED: use bUsePawnControlRotation instead */
 	UPROPERTY()
 	uint32 bUseControllerViewRotation_DEPRECATED:1;
 };
\ No newline at end of file
diff --git a/Engine/Source/Runtime/Engine/Classes/GameFramework/Pawn.h b/Engine/Source/Runtime/Engine/Classes/GameFramework/Pawn.h
--- a/Engine/Source/Runtime/Engine/Classes/GameFramework/Pawn.h
+++ b/Engine/Source/Runtime/Engine/Classes/GameFramework/Pawn.h
@@ -274,7 +274,10 @@ public:
 	UFUNCTION(BlueprintCallable, Category="Pawn")
 	virtual bool IsLocallyControlled() const;
 
-	/** @return the rotation the Pawn is looking */
+	/**
+	 * Get the view rotation of the Pawn (direction they are looking, normally Controller->ControlRotation).
+	 * @return The view rotation of the Pawn.
+	 */
 	virtual FRotator GetViewRotation() const;
 
 	/** @return	Pawn's eye location */
diff --git a/Engine/Source/Runtime/Engine/Classes/GameFramework/SpringArmComponent.h b/Engine/Source/Runtime/Engine/Classes/GameFramework/SpringArmComponent.h
--- a/Engine/Source/Runtime/Engine/Classes/GameFramework/SpringArmComponent.h
+++ b/Engine/Source/Runtime/Engine/Classes/GameFramework/SpringArmComponent.h
@@ -40,9 +40,12 @@ class ENGINE_API USpringArmComponent : public USceneComponent
 	UPROPERTY(EditAnywhere, BlueprintReadWrite, Category=CameraCollision)
 	uint32 bDoCollisionTest:1;
 
-	/** If this component is placed on a pawn, should it use the view rotation of the pawn where possible? */
+	/**
+	 * If this component is placed on a pawn, should it use the view/control rotation of the pawn where possible?
+	 * @see APawn::GetViewRotation()
+	 */
 	UPROPERTY(EditAnywhere, BlueprintReadWrite, Category=CameraSettings)
-	uint32 bUsePawnViewRotation:1;
+	uint32 bUsePawnControlRotation:1;
 
 	/** Should we inherit pitch from parent component. Does nothing if using Absolute Rotation. */
 	UPROPERTY(EditAnywhere, BlueprintReadWrite, Category=CameraSettings)
diff --git a/Engine/Source/Runtime/Engine/Private/Camera/CameraComponent.cpp b/Engine/Source/Runtime/Engine/Private/Camera/CameraComponent.cpp
--- a/Engine/Source/Runtime/Engine/Private/Camera/CameraComponent.cpp
+++ b/Engine/Source/Runtime/Engine/Private/Camera/CameraComponent.cpp
@@ -26,8 +26,8 @@ UCameraComponent::UCameraComponent(const class FPostConstructInitializePropertie
 	OrthoWidth = 512.0f;
 	bConstrainAspectRatio = false;
 	PostProcessBlendWeight = 1.0f;
-	bUseControllerViewRotation_DEPRECATED = true; // the previous default value before bUsePawnViewRotation replaced this var.
-	bUsePawnViewRotation = false;
+	bUseControllerViewRotation_DEPRECATED = true; // the previous default value before bUsePawnControlRotation replaced this var.
+	bUsePawnControlRotation = false;
 	bAutoActivate = true;
 }
 
@@ -96,7 +96,7 @@ void UCameraComponent::PostLoad()
 
 	if (LinkerUE4Ver < VER_UE4_RENAME_CAMERA_COMPONENT_VIEW_ROTATION)
 	{
-		bUsePawnViewRotation = bUseControllerViewRotation_DEPRECATED;
+		bUsePawnControlRotation = bUseControllerViewRotation_DEPRECATED;
 	}
 }
 
@@ -144,7 +144,7 @@ void UCameraComponent::PostEditChangeProperty(FPropertyChangedEvent& PropertyCha
 
 void UCameraComponent::GetCameraView(float DeltaTime, FMinimalViewInfo& DesiredView)
 {
-	if (bUsePawnViewRotation)
+	if (bUsePawnControlRotation)
 	{
 		if (APawn* OwningPawn = Cast<APawn>(GetOwner()))
 		{
diff --git a/Engine/Source/Runtime/Engine/Private/GameFramework/SpringArmComponent.cpp b/Engine/Source/Runtime/Engine/Private/GameFramework/SpringArmComponent.cpp
--- a/Engine/Source/Runtime/Engine/Private/GameFramework/SpringArmComponent.cpp
+++ b/Engine/Source/Runtime/Engine/Private/GameFramework/SpringArmComponent.cpp
@@ -17,7 +17,7 @@ USpringArmComponent::USpringArmComponent(const class FPostConstructInitializePro
 	
 	bAutoActivate = true;
 	bTickInEditor = true;
-	bUsePawnViewRotation = false;
+	bUsePawnControlRotation = false;
 	bDoCollisionTest = true;
 
 	bInheritPitch = true;
@@ -123,7 +123,7 @@ void USpringArmComponent::TickComponent(float DeltaTime, enum ELevelTick TickTyp
 {
 	Super::TickComponent(DeltaTime, TickType, ThisTickFunction);
 
-	if (bUsePawnViewRotation)
+	if (bUsePawnControlRotation)
 	{
 		if (APawn* OwningPawn = Cast<APawn>(GetOwner()))
 		{
