Hook up tutorials to relevant editors

diff --git a/Engine/Config/BaseEditorGameAgnostic.ini b/Engine/Config/BaseEditorGameAgnostic.ini
--- a/Engine/Config/BaseEditorGameAgnostic.ini
+++ b/Engine/Config/BaseEditorGameAgnostic.ini
@@ -15,4 +15,18 @@ bEnableWebCam=true
 WebCamResolution=Normal_640x480
 bMirrorWebCamImage=false
 bCaptureAudioFromComputer=true
-bCaptureAudioFromMicrophone=true
\ No newline at end of file
+bCaptureAudioFromMicrophone=true
+
+[/Script/IntroTutorials.EditorTutorialSettings]
++TutorialContexts=(Context="StaticMeshEditor",BrowserFilter=,AttractTutorial=None,LaunchTutorial=/Engine/Tutorial/SubEditors/StaticMeshEditorTutorial.StaticMeshEditorTutorial_C)
++TutorialContexts=(Context="LevelEditor",BrowserFilter=,AttractTutorial=None,LaunchTutorial=/Engine/Tutorial/Basics/LevelEditorOverview.LevelEditorOverview_C)
++TutorialContexts=(Context="BlueprintEditor",BrowserFilter=,AttractTutorial=None,LaunchTutorial=/Engine/Tutorial/BlueprintTutorials/BlueprintEditorTutorial.BlueprintEditorTutorial_C)
++TutorialContexts=(Context="BlueprintEditor.MacroLibrary",BrowserFilter=,AttractTutorial=None,LaunchTutorial=/Engine/Tutorial/BlueprintTutorials/BlueprintMacroLibrariesEditorOverview.BlueprintMacroLibrariesEditorOverview_C)
++TutorialContexts=(Context="BlueprintEditor.Interface",BrowserFilter=,AttractTutorial=None,LaunchTutorial=/Engine/Tutorial/BlueprintTutorials/BlueprintInterfacesEditorOverview.BlueprintInterfacesEditorOverview_C)
++TutorialContexts=(Context="BlueprintEditor.LevelScript",BrowserFilter=,AttractTutorial=None,LaunchTutorial=/Engine/Tutorial/BlueprintTutorials/LevelBlueprintEditorOverview.LevelBlueprintEditorOverview_C)
++TutorialContexts=(Context="DestructibleMeshEditor",BrowserFilter=,AttractTutorial=None,LaunchTutorial=/Engine/Tutorial/SubEditors/DestructibleMeshEditorTutorial.DestructibleMeshEditorTutorial_C)
++TutorialContexts=(Context="MaterialEditor",BrowserFilter=,AttractTutorial=None,LaunchTutorial=/Engine/Tutorial/SubEditors/MaterialEditorTutorial.MaterialEditorTutorial_C)
++TutorialContexts=(Context="Cascade",BrowserFilter=,AttractTutorial=None,LaunchTutorial=/Engine/Tutorial/SubEditors/ParticleSystemEditorTutorial.ParticleSystemEditorTutorial_C)
++TutorialContexts=(Context="Persona",BrowserFilter=,AttractTutorial=None,LaunchTutorial=/Engine/Tutorial/SubEditors/PersonaAnimEditorTutorial.PersonaAnimEditorTutorial_C)
++TutorialContexts=(Context="SpriteEditor",BrowserFilter=,AttractTutorial=None,LaunchTutorial=/Engine/Tutorial/Paper2D/SpriteEditorTutorial.SpriteEditorTutorial_C)
++TutorialContexts=(Context="FlipbookEditor",BrowserFilter=,AttractTutorial=None,LaunchTutorial=/Engine/Tutorial/Paper2D/FlipbookEditorTutorial.FlipbookEditorTutorial_C)
