#UE4 DOCS: adding tutorial categories

diff --git a/Engine/Config/BaseEditorGameAgnostic.ini b/Engine/Config/BaseEditorGameAgnostic.ini
--- a/Engine/Config/BaseEditorGameAgnostic.ini
+++ b/Engine/Config/BaseEditorGameAgnostic.ini
@@ -18,6 +18,13 @@ bCaptureAudioFromComputer=true
 bCaptureAudioFromMicrophone=true
 
 [/Script/IntroTutorials.EditorTutorialSettings]
++Categories=(Identifier="Basics",Title="Basics",Description="Getting started with Unreal Engine 4.",Icon="PlayWorld.RepeatLastPlay",Texture=None)
++Categories=(Identifier="Installation",Title="Installation",Description="Installing extra 3rd party software needed, or just useful to have, while working with Unreal Engine 4.",Icon="AssetEditor.SaveAsset",Texture=None)
++Categories=(Identifier="Mobile",Title="Mobile",Description="Mobile Tutorials.",Icon="MaterialEditor.ToggleMobileStats",Texture=None)
++Categories=(Identifier="Blueprints",Title="Blueprints",Description="Tutorials covering the creation and usage of Blueprints.",Icon="FullBlueprintEditor.EditGlobalOptions",Texture=/Engine/Tutorial/BlueprintTutorials/TutorialAssets/Blueprint_64x.Blueprint_64x)
++Categories=(Identifier="Paper2D",Title="Paper2D",Description="Tutorials convering the usage of Unreal Engine 4\'s 2D game tools: Paper2D.",Icon="ClassThumbnail.PaperSprite",Texture=/Engine/Tutorial/Paper2D/TutorialAssets/Paper2DSprite_TutorialIcon.Paper2DSprite_TutorialIcon)
++Categories=(Identifier="Landscape",Title="Landscape",Description="Tutorials covering the Unreal Editor 4 terrain editor: Landscape.",Icon="LevelEditor.LandscapeMode",Texture=/Engine/Tutorial/Landscape/TutorialAssets/Landscape.Landscape)
++Categories=(Identifier="Foliage",Title="Foliage",Description="Tutorials covering the Unreal Engine 4 Foliage tool.",Icon="LevelEditor.FoliageMode",Texture=/Engine/Tutorial/Foliage/TutorialAssets/Foliage.Foliage)
 +TutorialContexts=(Context="StaticMeshEditor",BrowserFilter=,AttractTutorial=None,LaunchTutorial=/Engine/Tutorial/SubEditors/StaticMeshEditorTutorial.StaticMeshEditorTutorial_C)
 +TutorialContexts=(Context="LevelEditor",BrowserFilter=,AttractTutorial=None,LaunchTutorial=/Engine/Tutorial/Basics/LevelEditorOverview.LevelEditorOverview_C)
 +TutorialContexts=(Context="BlueprintEditor",BrowserFilter=,AttractTutorial=None,LaunchTutorial=/Engine/Tutorial/BlueprintTutorials/BlueprintEditorTutorial.BlueprintEditorTutorial_C)
