#UE4 DOCS: Updates to the in-editor tutorials and the config file changes to make sure the categories stick.

diff --git a/Engine/Config/BaseEditorGameAgnostic.ini b/Engine/Config/BaseEditorGameAgnostic.ini
--- a/Engine/Config/BaseEditorGameAgnostic.ini
+++ b/Engine/Config/BaseEditorGameAgnostic.ini
@@ -25,6 +25,9 @@ bCaptureAudioFromMicrophone=true
 +Categories=(Identifier="Paper2D",Title="Paper2D",Description="Tutorials convering the usage of Unreal Engine 4\'s 2D game tools: Paper2D.",Icon="ClassThumbnail.PaperSprite",Texture=/Engine/Tutorial/Paper2D/TutorialAssets/Paper2DSprite_TutorialIcon.Paper2DSprite_TutorialIcon)
 +Categories=(Identifier="Landscape",Title="Landscape",Description="Tutorials covering the Unreal Editor 4 terrain editor: Landscape.",Icon="LevelEditor.LandscapeMode",Texture=/Engine/Tutorial/Landscape/TutorialAssets/Landscape.Landscape)
 +Categories=(Identifier="Foliage",Title="Foliage",Description="Tutorials covering the Unreal Engine 4 Foliage tool.",Icon="LevelEditor.FoliageMode",Texture=/Engine/Tutorial/Foliage/TutorialAssets/Foliage.Foliage)
++Categories=(Identifier="Editors",Title="Editors",Description="Tutorials covering the various sub editors of Unreal Engine 4.",Icon=,Texture=/Engine/Tutorial/SubEditors/TutorialAssets/icon_Editor_Preferences_General_40x.icon_Editor_Preferences_General_40x)
++Categories=(Identifier="Animation",Title="Animation",Description="Tutorials convering the animation system in Unreal Engine 4.",Icon=,Texture=/Engine/Tutorial/SubEditors/TutorialAssets/icon_ShowSkeletalMeshes_40x.icon_ShowSkeletalMeshes_40x)
++StartupTutorial=/Engine/Tutorial/Basics/LevelEditorAttract.LevelEditorAttract_C
 +TutorialContexts=(Context="StaticMeshEditor",BrowserFilter=,AttractTutorial=None,LaunchTutorial=/Engine/Tutorial/SubEditors/StaticMeshEditorTutorial.StaticMeshEditorTutorial_C)
 +TutorialContexts=(Context="LevelEditor",BrowserFilter=,AttractTutorial=None,LaunchTutorial=/Engine/Tutorial/Basics/LevelEditorOverview.LevelEditorOverview_C)
 +TutorialContexts=(Context="BlueprintEditor",BrowserFilter=,AttractTutorial=None,LaunchTutorial=/Engine/Tutorial/BlueprintTutorials/BlueprintEditorTutorial.BlueprintEditorTutorial_C)
