Remove duplicate redirects

diff --git a/Engine/Config/BaseEngine.ini b/Engine/Config/BaseEngine.ini
--- a/Engine/Config/BaseEngine.ini
+++ b/Engine/Config/BaseEngine.ini
@@ -1207,10 +1207,6 @@ bUseStreamingPause=false
 +TaggedPropertyRedirects=(ClassName="GameState",OldPropertyName="GameClass",NewPropertyName="GameModeClass")
 +TaggedPropertyRedirects=(ClassName="GameReplicationInfo",OldPropertyName="GameClass",NewPropertyName="GameModeClass")
 
-+ActiveClassRedirects=(OldClassName="AnimGraphNode_BlendSpace",NewClassName="/Script/AnimGraph.AnimGraphNode_BlendSpacePlayer")
-+ActiveClassRedirects=(OldClassName="AnimNode_BlendSpace",NewClassName="/Script/Engine.AnimNode_BlendSpacePlayer")
-+ActiveStructRedirects=(OldStructName="AnimNode_BlendSpace",NewStructName="AnimNode_BlendSpacePlayer")
-
 +ActiveStructRedirects=(OldStructName="FFormatTextArgument",NewStructName="/Script/Engine.FFormatTextArgument")
 
 +K2FieldRedirects=(OldFieldName="AnimInstance.GetSequencePlayerLength",NewFieldName="GetAnimAssetPlayerLength")
