Created LOD Group for 2D Pixels, which disables mipmaps and filtering (added missing file)
#ttp 345443 - NUX:Texture Editor: "Nearest" texture filtering option is buried away, but this is super useful for 2D games
#branch UE4
#proj Runtime.Engine

diff --git a/Engine/Config/BaseEngine.ini b/Engine/Config/BaseEngine.ini
--- a/Engine/Config/BaseEngine.ini
+++ b/Engine/Config/BaseEngine.ini
@@ -1786,6 +1786,7 @@ TEXTUREGROUP_MobileFlattened=(MinLODSize=1,MaxLODSize=4096,LODBias=0,MinMagFilte
 TEXTUREGROUP_Terrain_Heightmap=(MinLODSize=1,MaxLODSize=4096,LODBias=0,MinMagFilter=aniso,MipFilter=point)
 TEXTUREGROUP_Terrain_Weightmap=(MinLODSize=1,MaxLODSize=4096,LODBias=0,MinMagFilter=aniso,MipFilter=point)
 TEXTUREGROUP_Bokeh=(MinLODSize=1,MaxLODSize=256,LODBias=0,MinMagFilter=linear,MipFilter=linear)
+TEXTUREGROUP_Pixels2D=(MinLODSize=1,MaxLODSize=4096,LODBias=0,MinMagFilter=point,MipFilter=point)
 
 [SystemSettingsEditor]
 ; System settings overrides for the editor.  Ideally the editor should use the same settings as the game.
