Tweaked lightmap density values

diff --git a/Engine/Config/BaseEngine.ini b/Engine/Config/BaseEngine.ini
--- a/Engine/Config/BaseEngine.ini
+++ b/Engine/Config/BaseEngine.ini
@@ -133,8 +133,8 @@ CameraTranslationThreshold=10000
 PrimitiveProbablyVisibleTime = 8.0
 MaxOcclusionPixelsFraction = 0.1
 MinLightMapDensity=0.0
-IdealLightMapDensity=0.05
-MaxLightMapDensity=0.2
+IdealLightMapDensity=0.2
+MaxLightMapDensity=0.8
 RenderLightMapDensityGrayscaleScale=1.0
 RenderLightMapDensityColorScale=1.0
 bRenderLightMapDensityGrayscale=false
