UE4 - auto complete, comments

diff --git a/Engine/Config/BaseInput.ini b/Engine/Config/BaseInput.ini
--- a/Engine/Config/BaseInput.ini
+++ b/Engine/Config/BaseInput.ini
@@ -133,12 +133,12 @@ MaxScrollbackSize=1024
 +ManualAutoCompleteList=(Command="IsolateReverb",Desc="IsolateReverb")
 +ManualAutoCompleteList=(Command="ResetSoundState",Desc="ResetSoundState (Resets volumes to default and removes test filters)")
 +ManualAutoCompleteList=(Command="ModifySoundClass",Desc="ModifySoundClass <SoundClassName> Vol=<new volume>")
-+ManualAutoCompleteList=(Command="DisableAllScreenMessages",Desc="Disables all on-screen warnings/messages")
-+ManualAutoCompleteList=(Command="EnableAllScreenMessages",Desc="Enables all on-screen warnings/messages")
-+ManualAutoCompleteList=(Command="ToggleAllScreenMessages",Desc="Toggles display state of all on-screen warnings/messages")
-+ManualAutoCompleteList=(Command="ToggleRHIThread",Desc="Toggles the RHI thread for platforms that have it")
-+ManualAutoCompleteList=(Command="ToggleRenderingThread",Desc="Toggles the rendering thread for platforms that have it")
-+ManualAutoCompleteList=(Command="CaptureMode",Desc="Toggles display state of all on-screen warnings/messages")
++ManualAutoCompleteList=(Command="DisableAllScreenMessages",Desc="DisableAllScreenMessages (Disables all on-screen warnings/messages)")
++ManualAutoCompleteList=(Command="EnableAllScreenMessages",Desc="EnableAllScreenMessages (Enables all on-screen warnings/messages)")
++ManualAutoCompleteList=(Command="ToggleAllScreenMessages",Desc="ToggleAllScreenMessages (Toggles display state of all on-screen warnings/messages)")
++ManualAutoCompleteList=(Command="ToggleRHIThread",Desc="ToggleRHIThread (Toggles the RHI thread for platforms that have it)")
++ManualAutoCompleteList=(Command="ToggleRenderingThread",Desc="ToggleRenderingThread (Toggles the rendering thread for platforms that have it)")
++ManualAutoCompleteList=(Command="CaptureMode",Desc="CaptureMode (Toggles display state of all on-screen warnings/messages)")
 +ManualAutoCompleteList=(Command="ShowDebug None",Desc="ShowDebug None (Toggles ShowDebug w/ current debug type selection)")
 +ManualAutoCompleteList=(Command="ShowDebug Reset",Desc="ShowDebug Reset (Turns off ShowDebug, and clears debug type selection)")
 +ManualAutoCompleteList=(Command="ShowDebug NET",Desc="ShowDebug NET")
diff --git a/Engine/Source/Runtime/Renderer/Private/DeferredShadingRenderer.cpp b/Engine/Source/Runtime/Renderer/Private/DeferredShadingRenderer.cpp
--- a/Engine/Source/Runtime/Renderer/Private/DeferredShadingRenderer.cpp
+++ b/Engine/Source/Runtime/Renderer/Private/DeferredShadingRenderer.cpp
@@ -675,8 +675,6 @@ void FDeferredShadingSceneRenderer::Render(FRHICommandListImmediate& RHICmdList)
 	if (GRHIThread)
 	{
 		// we will probably stall on occlusion queries, so might as well have the RHI thread and GPU work while we wait.
-		//QUICK_SCOPE_CYCLE_COUNTER(STAT_OcclusionCull_Dispatch);
-		//FRHICommandListExecutor::GetImmediateCommandList().ImmediateFlush(EImmediateFlushType::DispatchToRHIThread); // we want to make sure this all gets to the GPU this frame and doesn't hang around
 		QUICK_SCOPE_CYCLE_COUNTER(STAT_PostInitViews_FlushDel);
 		FRHICommandListExecutor::GetImmediateCommandList().ImmediateFlush(EImmediateFlushType::FlushRHIThreadFlushResources);
 	}
@@ -869,7 +867,7 @@ void FDeferredShadingSceneRenderer::Render(FRHICommandListImmediate& RHICmdList)
 		QUICK_SCOPE_CYCLE_COUNTER(STAT_OcclusionSubmittedFence_Dispatch);
 		OcclusionSubmittedFence = FRHICommandListExecutor::RHIThreadFence();
 	}
-	
+
 	// Render lighting.
 	if (ViewFamily.EngineShowFlags.Lighting
 		&& FeatureLevel >= ERHIFeatureLevel::SM4
diff --git a/Engine/Source/Runtime/Renderer/Private/TranslucentRendering.cpp b/Engine/Source/Runtime/Renderer/Private/TranslucentRendering.cpp
--- a/Engine/Source/Runtime/Renderer/Private/TranslucentRendering.cpp
+++ b/Engine/Source/Runtime/Renderer/Private/TranslucentRendering.cpp
@@ -892,12 +892,12 @@ static TAutoConsoleVariable<int32> CVarRHICmdTranslucencyPassDeferredContexts(
 void FDeferredShadingSceneRenderer::RenderTranslucencyParallel(FRHICommandListImmediate& RHICmdList)
 {
 	GSceneRenderTargets.AllocLightAttenuation(); // materials will attempt to get this texture before the deferred command to set it up executes
+	check(IsInRenderingThread());
+
 
 	FScopedCommandListWaitForTasks Flusher(RHICmdList);
 	FTranslucencyDrawingPolicyFactory::ContextType ThisContext;
 
-	check(IsInRenderingThread());
-
 	for (int32 ViewIndex = 0; ViewIndex < Views.Num(); ViewIndex++)
 	{
 		SCOPED_CONDITIONAL_DRAW_EVENTF(RHICmdList, EventView, Views.Num() > 1, DEC_SCENE_ITEMS, TEXT("View%d"), ViewIndex);
@@ -905,7 +905,6 @@ void FDeferredShadingSceneRenderer::RenderTranslucencyParallel(FRHICommandListIm
 		const FViewInfo& View = Views[ViewIndex];
 
 		{
-			//@todo multiple views would be better moving this outside of the loop to batch all views together
 			FTranslucencyPassParallelCommandListSet ParallelCommandListSet(View, RHICmdList, nullptr, CVarRHICmdTranslucencyPassDeferredContexts.GetValueOnRenderThread() > 0, false);
 
 			{
