Integrated GitHub pull request #487

Increases range of power of two grid sizes

diff --git a/Engine/Config/BaseEditorUserSettings.ini b/Engine/Config/BaseEditorUserSettings.ini
--- a/Engine/Config/BaseEditorUserSettings.ini
+++ b/Engine/Config/BaseEditorUserSettings.ini
@@ -289,6 +289,9 @@ bUsePowerOf2SnapSize=False
 .Pow2GridSizes=256
 .Pow2GridSizes=512
 .Pow2GridSizes=1024
+.Pow2GridSizes=2048
+.Pow2GridSizes=4096
+.Pow2GridSizes=8192
 .Pow2GridIntervals=8
 
 .CommonRotGridSizes=5
