removed diffuse and specular visuzlisebuffr from engine (no longer needed as we fully switched to basecolor/metallic/specular)

diff --git a/Engine/Config/BaseEngine.ini b/Engine/Config/BaseEngine.ini
--- a/Engine/Config/BaseEngine.ini
+++ b/Engine/Config/BaseEngine.ini
@@ -2233,7 +2233,6 @@ EnableNiagara=false
 BaseColor=(Material="/Engine/BufferVisualization/BaseColor.BaseColor", Name=LOCTEXT("BaseColorMat", "Base Color"))
 CustomDepth=(Material="/Engine/BufferVisualization/CustomDepth.CustomDepth", Name=LOCTEXT("BaseCustomDepthMat", "Custom Depth"))
 DecalMask=(Material="/Engine/BufferVisualization/DecalMask.DecalMask", Name=LOCTEXT("BaseDecalMaskMat", "Decal Mask"))
-DiffuseColor=(Material="/Engine/BufferVisualization/DiffuseColor.DiffuseColor", Name=LOCTEXT("BaseDiffuseColorMat", "Diffuse Color"))
 FinalImage=(Material="/Engine/BufferVisualization/FinalImage.FinalImage", Name=LOCTEXT("BaseFinalImageMat", "Final Image"))
 ShadingModel=(Material="/Engine/BufferVisualization/LightingModel.LightingModel", Name=LOCTEXT("BaseShadingModelMat", "Shading Model"))
 MaterialAO=(Material="/Engine/BufferVisualization/MaterialAO.MaterialAO", Name=LOCTEXT("BaseMaterialAOMat", "Material Ambient Occlusion"))
@@ -2244,7 +2243,6 @@ SceneColor=(Material="/Engine/BufferVisualization/SceneColor.SceneColor", Name=L
 SceneDepth=(Material="/Engine/BufferVisualization/SceneDepth.SceneDepth", Name=LOCTEXT("BaseSceneDepthMat", "Scene Depth"))
 SeparateTranslucencyRGB=(Material="/Engine/BufferVisualization/SeparateTranslucencyRGB.SeparateTranslucencyRGB", Name=LOCTEXT("BaseSeparateTranslucencyRGBMat", "Separate Translucency RGB"))
 SeparateTranslucencyA=(Material="/Engine/BufferVisualization/SeparateTranslucencyA.SeparateTranslucencyA", Name=LOCTEXT("BaseSeparateTranslucencyAMat", "Separate Translucency Alpha"))
-SpecularColor=(Material="/Engine/BufferVisualization/SpecularColor.SpecularColor", Name=LOCTEXT("BaseSpecularColorMat", "Specular Color"))
 Specular=(Material="/Engine/BufferVisualization/Specular.Specular", Name=LOCTEXT("BaseSpecularMat", "Specular"))
 SubsurfaceColor=(Material="/Engine/BufferVisualization/SubsurfaceColor.SubsurfaceColor", Name=LOCTEXT("BaseSubsurfaceColorMat", "Subsurface Color"))
 WorldNormal=(Material="/Engine/BufferVisualization/WorldNormal.WorldNormal", Name=LOCTEXT("BaseWorldNormalMat", "World Normal"))
