fix for resolution issue on some of the sample games
UE-3126, UE-2918
#ios

diff --git a/Engine/Config/BaseDeviceProfiles.ini b/Engine/Config/BaseDeviceProfiles.ini
--- a/Engine/Config/BaseDeviceProfiles.ini
+++ b/Engine/Config/BaseDeviceProfiles.ini
@@ -43,6 +43,7 @@ BaseProfileName=Windows
 [IOS DeviceProfile]
 DeviceType=IOS
 BaseProfileName=
++CVars=r.MobileContentScaleFactor=0
 +CVars=r.BloomQuality=0
 +CVars=r.DepthOfFieldQuality=0
 +CVars=r.LightShaftQuality=0
@@ -108,7 +109,6 @@ BaseProfileName=IOS
 [iPhone6 DeviceProfile]
 DeviceType=IOS
 BaseProfileName=IOS
-+CVars=r.MobileContentScaleFactor=0
 +CVars=r.BloomQuality=1
 +CVars=r.DepthOfFieldQuality=1
 +CVars=r.LightShaftQuality=1
@@ -116,7 +116,6 @@ BaseProfileName=IOS
 [iPhone6Plus DeviceProfile]
 DeviceType=IOS
 BaseProfileName=IOS
-+CVars=r.MobileContentScaleFactor=0
 +CVars=r.BloomQuality=1
 +CVars=r.DepthOfFieldQuality=1
 +CVars=r.LightShaftQuality=1
