Added a bunch of ini settings which effect cooked content.

diff --git a/Engine/Config/BaseEditor.ini b/Engine/Config/BaseEditor.ini
--- a/Engine/Config/BaseEditor.ini
+++ b/Engine/Config/BaseEditor.ini
@@ -289,4 +289,67 @@ HighSpawnRateOrBurstThreshold=35
 FarLODDistanceTheshold=3000
 
 [/Script/UnrealEd.CookCommandlet]
-+FullGCAssetClassNames=/Script/Engine.World
\ No newline at end of file
++FullGCAssetClassNames=/Script/Engine.World
+
+[CookSettings]
++VersionedIniParams=Engine:/Script/Engine.RendererSettings:Compat.UseDXT5NormalMaps
++VersionedIniParams=Engine:/Script/AndroidRuntimeSettings.AndroidRuntimeSettings:AndroidAudio
++VersionedIniParams=Engine:/Script/UnrealEd.CookerSettings:DefaultPVRTCQuality
++VersionedIniParams=Editor:StructSerialization:BreakSerializationRecursion
++VersionedIniParams=EditorUserSettings:/Script/UnrealEd.EditorUserSettings:MaterialQualityLevel
++VersionedIniParams=Engine:/Script/Engine.Engine:GameUserSettingsClassName
++VersionedIniParams=Engine:/Script/EngineSettings.GameMapsSettings:Compat.UseDXT5NormalMaps
++VersionedIniParams=Engine:/Script/EngineSettings.GameMapsSettings:GameDefaultMap
++VersionedIniParams=Engine:/Script/EngineSettings.GameMapsSettings:GameInstanceClass
++VersionedIniParams=Engine:/Script/EngineSettings.GameMapsSettings:ServerDefaultMap
+
++VersionedIniParams=Engine:Paper2D:PaperAxisX
++VersionedIniParams=Engine:Paper2D:PaperAxisY
+
+; any default classes / default assets will not be updated unless we recook
++VersionedIniParams=Input:/Script/Engine.InputSettings:DefaultTouchInterface
++VersionedIniParams=Engine:/Script/Engine.Engine:DefaultTextureName
++VersionedIniParams=Engine:/Script/Engine.Engine:ConsoleClassName
++VersionedIniParams=Engine:/Script/Engine.Engine:GameViewportClientClassName
++VersionedIniParams=Engine:/Script/Engine.Engine:LocalPlayerClassName
++VersionedIniParams=Engine:/Script/Engine.Engine:WorldSettingsClassName
++VersionedIniParams=Engine:/Script/Engine.Engine:AvoidanceManagerClassName
++VersionedIniParams=Engine:/Script/Engine.Engine:NavigationSystemClassName
++VersionedIniParams=Engine:/Script/Engine.Engine:PhysicsCollisionHandlerClassName
++VersionedIniParams=Engine:/Script/Engine.Engine:AIControllerClassName
++VersionedIniParams=Engine:/Script/Engine.Engine:LevelScriptActorClassName
++VersionedIniParams=Engine:/Script/Engine.Engine:GameSingletonClassName
++VersionedIniParams=Engine:/Script/Engine.Engine:DefaultBlueprintBaseClassName
++VersionedIniParams=Engine:/Script/Engine.Engine:DefaultPreviewPawnClassName
++VersionedIniParams=Engine:/Script/Engine.Engine:DefaultTireTypeName
+
++VersionedIniParams=Engine:/Script/Engine.Engine:DefaultDiffuseTextureName
++VersionedIniParams=Engine:/Script/Engine.Engine:HighFrequencyNoiseTextureName
++VersionedIniParams=Engine:/Script/Engine.Engine:DefaultBSPVertexTextureName
++VersionedIniParams=Engine:/Script/Engine.Engine:WireframeMaterialName
++VersionedIniParams=Engine:/Script/Engine.Engine:DefaultBokehTextureName
++VersionedIniParams=Engine:/Script/Engine.Engine:DebugMeshMaterialName
++VersionedIniParams=Engine:/Script/Engine.Engine:LevelColorationLitMaterialName
++VersionedIniParams=Engine:/Script/Engine.Engine:LevelColorationUnlitMaterialName
++VersionedIniParams=Engine:/Script/Engine.Engine:LightingTexelDensityName
++VersionedIniParams=Engine:/Script/Engine.Engine:ShadedLevelColorationLitMaterialName
++VersionedIniParams=Engine:/Script/Engine.Engine:ShadedLevelColorationUnlitMaterialName
++VersionedIniParams=Engine:/Script/Engine.Engine:RemoveSurfaceMaterialName
++VersionedIniParams=Engine:/Script/Engine.Engine:VertexColorMaterialName
++VersionedIniParams=Engine:/Script/Engine.Engine:VertexColorViewModeMaterialName_ColorOnly
++VersionedIniParams=Engine:/Script/Engine.Engine:VertexColorViewModeMaterialName_AlphaAsColor
++VersionedIniParams=Engine:/Script/Engine.Engine:VertexColorViewModeMaterialName_RedOnly
++VersionedIniParams=Engine:/Script/Engine.Engine:VertexColorViewModeMaterialName_BlueOnly
++VersionedIniParams=Engine:/Script/Engine.Engine:VertexColorViewModeMaterialName_GreenOnly
++VersionedIniParams=Engine:/Script/Engine.Engine:InvalidLightmapSettingsMaterialName
++VersionedIniParams=Engine:/Script/Engine.Engine:ConstraintLimitMaterialName
++VersionedIniParams=Engine:/Script/Engine.Engine:PreviewShadowsIndicatorMaterialName
++VersionedIniParams=Engine:/Script/Engine.Engine:ArrowMaterialName
++VersionedIniParams=Engine:/Script/Engine.Engine:DefaultPhysMaterialName
++VersionedIniParams=Engine:/Script/Engine.Engine:PreIntegratedSkinBRDFTextureName
++VersionedIniParams=Engine:/Script/Engine.Engine:MiniFontTextureName
++VersionedIniParams=Engine:/Script/Engine.Engine:WeightMapPlaceholderTextureName
++VersionedIniParams=Engine:/Script/Engine.Engine:LightMapDensityTextureName
+
+
+
