Remove push to talk from debug exec bindings

diff --git a/Engine/Config/BaseInput.ini b/Engine/Config/BaseInput.ini
--- a/Engine/Config/BaseInput.ini
+++ b/Engine/Config/BaseInput.ini
@@ -28,7 +28,6 @@ bRequireCtrlToNavigateAutoComplete=False
 +DebugExecBindings=(Key=Pause,Command="Pause")
 +DebugExecBindings=(Key=F11,Command="LevelEditor.ToggleImmersive")
 +DebugExecBindings=(Key=F11,Command="MainFrame.ToggleFullscreen",Shift=True)
-+DebugExecBindings=(Key=B,Command="ToggleSpeaking true | OnRelease ToggleSpeaking false")
 +DebugExecBindings=(Key=F1,Command="ShowMouseCursor",Shift=True)
 +DebugExecBindings=(Key=F1,Command="viewmode wireframe", bIgnoreShift=True)
 +DebugExecBindings=(Key=F2,Command="viewmode unlit")
@@ -42,7 +41,7 @@ bRequireCtrlToNavigateAutoComplete=False
 +DebugExecBindings=(Key=Enter,Command="TOGGLE_FULLSCREEN",Alt=True)
 +DebugExecBindings=(Key=Tab,Command="FocusNextPIEWindow",Control=True)
 +DebugExecBindings=(Key=Tab,Command="FocusLastPIEWindow",Control=True,Shift=True)
-+DebugExecBindings=(Key=Apostrophe,Command="EnableGDT",Control=False,Shift=False)
++DebugExecBindings=(Key=Apostrophe,Command="EnableGDT")
 
 [/Script/Engine.Console]
 HistoryBot=-1
