Swapped Canvas over to use the Roboto font assets by default

diff --git a/Engine/Config/BaseEngine.ini b/Engine/Config/BaseEngine.ini
--- a/Engine/Config/BaseEngine.ini
+++ b/Engine/Config/BaseEngine.ini
@@ -43,11 +43,11 @@ bAllowMatureLanguage=false
 GameEngine=/Script/Engine.GameEngine
 EditorEngine=/Script/UnrealEd.EditorEngine
 UnrealEdEngine=/Script/UnrealEd.UnrealEdEngine
-TinyFontName=/Engine/EngineFonts/TinyFont.TinyFont
-SmallFontName=/Engine/EngineFonts/SmallFont.SmallFont
-MediumFontName=/Engine/EngineFonts/SmallFont.SmallFont
-LargeFontName=/Engine/EngineFonts/SmallFont.SmallFont
-SubtitleFontName=/Engine/EngineFonts/SmallFont.SmallFont
+TinyFontName=/Engine/EngineFonts/RobotoTiny.RobotoTiny
+SmallFontName=/Engine/EngineFonts/Roboto.Roboto
+MediumFontName=/Engine/EngineFonts/Roboto.Roboto
+LargeFontName=/Engine/EngineFonts/Roboto.Roboto
+SubtitleFontName=/Engine/EngineFonts/Roboto.Roboto
 WireframeMaterialName=/Engine/EngineDebugMaterials/WireframeMaterial.WireframeMaterial
 DefaultMaterialName=/Engine/EngineMaterials/WorldGridMaterial.WorldGridMaterial
 DefaultLightFunctionMaterialName=/Engine/EngineMaterials/DefaultLightFunctionMaterial.DefaultLightFunctionMaterial
@@ -1462,6 +1462,8 @@ bFullyCompressStartupPackages=false
 
 +Package=/Engine/EngineFonts/SmallFont
 +Package=/Engine/EngineFonts/TinyFont
++Package=/Engine/EngineFonts/Roboto
++Package=/Engine/EngineFonts/RobotoTiny
 
 ; only needed for TextRender feature (3d Text in world)
 +Package=/Engine/EngineMaterials/DefaultTextMaterialTranslucent
@@ -1659,6 +1661,8 @@ AutoSaveIndex=0
 
 +PackagesToBeFullyLoadedAtStartup=/Engine/EngineFonts/SmallFont
 +PackagesToBeFullyLoadedAtStartup=/Engine/EngineFonts/TinyFont
++PackagesToBeFullyLoadedAtStartup=/Engine/EngineFonts/Roboto
++PackagesToBeFullyLoadedAtStartup=/Engine/EngineFonts/RobotoTiny
 
 +PackagesToBeFullyLoadedAtStartup=/Engine/EngineResources/Black
 +PackagesToBeFullyLoadedAtStartup=/Engine/EngineResources/DefaultTexture
diff --git a/Engine/Source/Editor/UnrealEd/Private/Editor.cpp b/Engine/Source/Editor/UnrealEd/Private/Editor.cpp
--- a/Engine/Source/Editor/UnrealEd/Private/Editor.cpp
+++ b/Engine/Source/Editor/UnrealEd/Private/Editor.cpp
@@ -207,7 +207,7 @@ UEditorEngine::UEditorEngine(const class FObjectInitializer& ObjectInitializer)
 				, EditorSphereMesh(TEXT("/Engine/EditorMeshes/EditorSphere"))
 				, EditorPlaneMesh(TEXT("/Engine/EditorMeshes/EditorPlane"))
 				, EditorCylinderMesh(TEXT("/Engine/EditorMeshes/EditorCylinder"))
-				, SmallFont(TEXT("/Engine/EditorResources/SmallFont"))
+				, SmallFont(TEXT("/Engine/EngineFonts/Roboto"))
 			{
 			}
 		};
