UMG - Adding several parameter redirectors to help convert more functions in UMG that have been updated to support controller input/navigation from Slate.

diff --git a/Engine/Config/BaseEngine.ini b/Engine/Config/BaseEngine.ini
--- a/Engine/Config/BaseEngine.ini
+++ b/Engine/Config/BaseEngine.ini
@@ -1349,6 +1349,19 @@ bUseStreamingPause=false
 +K2ParamRedirects=(NodeName="/Script/BlueprintGraph.K2Node_FunctionEntry", OldParamName="AnimNotifyState.Received_NotifyTick.AnimSeq", NewParamName="Animation")
 +K2ParamRedirects=(NodeName="/Script/BlueprintGraph.K2Node_FunctionEntry", OldParamName="AnimNotifyState.Received_NotifyEnd.AnimSeq", NewParamName="Animation")
 
+; UMG Controller Input Changes
+
++K2ParamRedirects=(NodeName="/Script/BlueprintGraph.K2Node_FunctionEntry", OldParamName="UserWidget.OnKeyboardFocusReceived.InKeyboardFocusEvent", NewParamName="InFocusEvent")
++K2ParamRedirects=(NodeName="/Script/BlueprintGraph.K2Node_Event", OldParamName="UserWidget.OnKeyboardFocusLost.InKeyboardFocusEvent", NewParamName="InFocusEvent")
++K2ParamRedirects=(NodeName="/Script/BlueprintGraph.K2Node_FunctionEntry", OldParamName="UserWidget.OnPreviewKeyDown.InKeyboardEvent", NewParamName="InKeyEvent")
++K2ParamRedirects=(NodeName="/Script/BlueprintGraph.K2Node_FunctionEntry", OldParamName="UserWidget.OnKeyDown.InKeyboardEvent", NewParamName="InKeyEvent")
++K2ParamRedirects=(NodeName="/Script/BlueprintGraph.K2Node_FunctionEntry", OldParamName="UserWidget.OnKeyUp.InKeyboardEvent", NewParamName="InKeyEvent")
++K2ParamRedirects=(NodeName="/Script/BlueprintGraph.K2Node_FunctionEntry", OldParamName="UserWidget.OnControllerAnalogValueChanged.ControllerEvent", NewParamName="InAnalogInputEvent")
++K2ParamRedirects=(NodeName="/Script/BlueprintGraph.K2Node_FunctionEntry", OldParamName="UserWidget.OnControllerButtonPressed.ControllerEvent", NewParamName="InKeyEvent")
++K2ParamRedirects=(NodeName="/Script/BlueprintGraph.K2Node_FunctionEntry", OldParamName="UserWidget.OnControllerButtonReleased.ControllerEvent", NewParamName="InKeyEvent")
+
+
+
 ; VER_UE4_INVERSE_SQUARED_LIGHTS_DEFAULT
 +K2FieldRedirects=(OldFieldName="LightComponentBase.Brightness",NewFieldName="LightComponentBase.Intensity")
 +K2FieldRedirects=(OldFieldName="PointLightComponent.Radius",NewFieldName="PointLightComponent.AttenuationRadius")
