fix for base profile using MobileContentScaleFactor of 0 by default
re-fixed Black Jack and Memory Game
#ios

diff --git a/Engine/Config/BaseDeviceProfiles.ini b/Engine/Config/BaseDeviceProfiles.ini
--- a/Engine/Config/BaseDeviceProfiles.ini
+++ b/Engine/Config/BaseDeviceProfiles.ini
@@ -43,7 +43,6 @@ BaseProfileName=Windows
 [IOS DeviceProfile]
 DeviceType=IOS
 BaseProfileName=
-+CVars=r.MobileContentScaleFactor=0
 +CVars=r.BloomQuality=0
 +CVars=r.DepthOfFieldQuality=0
 +CVars=r.LightShaftQuality=0
