ENABLED ReinstanceOnlyWhenNecessary:
- mechanism to distinguish if instances of recompiled BPGC must be reinstanced
- replacing PB class without reinstancing.

diff --git a/Engine/Config/BaseEngine.ini b/Engine/Config/BaseEngine.ini
--- a/Engine/Config/BaseEngine.ini
+++ b/Engine/Config/BaseEngine.ini
@@ -1723,7 +1723,7 @@ bTurnOffEditorConstructionScript=false
 PrintStringDuration=2.0
 bUseLocalGraphVariables=false
 bPersistentUberGraphFrame=false
-bReinstanceOnlyWhenNecessary=false
+bReinstanceOnlyWhenNecessary=true
 bDisplaysLayout=false
 bNativeCodeGenerationTool=false
 
