ENABLED Uber Graph Persistent Frame - All local variables, used by event graphs (see FKismetFunctionContext::EventGraphLocals), are no longer stored in class  scope. They are stored on a frame (as local variables in regular functions). For the ExecuteUbergraph_*() functions the FFrameÆs allocation is kept for the life of the object.

diff --git a/Engine/Config/BaseEngine.ini b/Engine/Config/BaseEngine.ini
--- a/Engine/Config/BaseEngine.ini
+++ b/Engine/Config/BaseEngine.ini
@@ -1722,7 +1722,7 @@ CompileDisplaysAnimBlueprintBackend=false
 bTurnOffEditorConstructionScript=false
 PrintStringDuration=2.0
 bUseLocalGraphVariables=false
-bPersistentUberGraphFrame=false
+bPersistentUberGraphFrame=true
 bReinstanceOnlyWhenNecessary=true
 bDisplaysLayout=false
 bNativeCodeGenerationTool=false
