BP: After (only) a default value was changed, the recompilation of BP doesn't cause reinstancing.

diff --git a/Engine/Config/BaseEngine.ini b/Engine/Config/BaseEngine.ini
--- a/Engine/Config/BaseEngine.ini
+++ b/Engine/Config/BaseEngine.ini
@@ -1727,7 +1727,7 @@ PrintStringDuration=2.0
 bUseLocalGraphVariables=false
 bPersistentUberGraphFrame=true
 bReinstanceOnlyWhenNecessary=true
-bChangeDefaultValueWithoutReinstancing=false
+bChangeDefaultValueWithoutReinstancing=true
 bDisplaysLayout=false
 bNativeCodeGenerationTool=false
 
