r.Shaders.KeepDebugInfo comment

diff --git a/Engine/Config/ConsoleVariables.ini b/Engine/Config/ConsoleVariables.ini
--- a/Engine/Config/ConsoleVariables.ini
+++ b/Engine/Config/ConsoleVariables.ini
@@ -26,3 +26,7 @@
 ;r.DumpShaderDebugInfo=1
 ; Uncomment to enable parallel rendering at startup
 ;r.RHICmdBypass=0
+
+; Uncomment when running with a graphical debugger (but not when profiling)
+;r.Shaders.Optimize=0
+;r.Shaders.KeepDebugInfo=1
\ No newline at end of file
