UE-5501 LIVE: StartFPSChart Command references an invalid command in the console

diff --git a/Engine/Config/BaseInput.ini b/Engine/Config/BaseInput.ini
--- a/Engine/Config/BaseInput.ini
+++ b/Engine/Config/BaseInput.ini
@@ -51,7 +51,7 @@ MaxScrollbackSize=1024
 +AutoCompleteMapPaths=Content/Maps
 +ManualAutoCompleteList=(Command="Exit",Desc="Exit (Exits the game)")
 +ManualAutoCompleteList=(Command="DebugCreatePlayer 1",Desc=)
-+ManualAutoCompleteList=(Command="StartFPSChart",Desc="StartFPSChart (after that use EndFPSChart)")
++ManualAutoCompleteList=(Command="StartFPSChart",Desc="StartFPSChart (after that use StopFPSChart)")
 +ManualAutoCompleteList=(Command="StopFPSChart",Desc="StopFPSChart (after that look for the output in Saved/Profiling/FPSChartStats)")
 +ManualAutoCompleteList=(Command="FreezeAt",Desc="Locks the player view and rendering time.")
 +ManualAutoCompleteList=(Command="Open",Desc="Open <MapName> (Opens the specified map, doesn't pass previously set options)")
