- Cleaned up CollisionChannel to not display deprecated channels
- Also removed dynamic/static deprecated channel and use one deprecated one

diff --git a/Engine/Config/BaseEngine.ini b/Engine/Config/BaseEngine.ini
--- a/Engine/Config/BaseEngine.ini
+++ b/Engine/Config/BaseEngine.ini
@@ -737,8 +737,10 @@ bUseStreamingPause=false
 +EnumRedirects=(EnumName="ECollisionChannel",OldEnumEntry="ECC_TouchAllDynamic",NewEnumEntry="ECC_OverlapAllDynamic")
 +EnumRedirects=(EnumName="ECollisionChannel",OldEnumEntry="ECC_TouchAllStatic",NewEnumEntry="ECC_OverlapAllStatic")
 +EnumRedirects=(EnumName="ECollisionChannel",OldEnumEntry="ECC_OverlapAll",NewEnumEntry="ECC_OverlapAll_Deprecated")
-+EnumRedirects=(EnumName="ECollisionChannel",OldEnumEntry="ECC_OverlapAllDynamic",NewEnumEntry="ECC_OverlapAllDynamic_Deprecated")
-+EnumRedirects=(EnumName="ECollisionChannel",OldEnumEntry="ECC_OverlapAllStatic",NewEnumEntry="ECC_OverlapAllStatic_Deprecated")
++EnumRedirects=(EnumName="ECollisionChannel",OldEnumEntry="ECC_OverlapAllDynamic",NewEnumEntry="ECC_OverlapAll_Deprecated")
++EnumRedirects=(EnumName="ECollisionChannel",OldEnumEntry="ECC_OverlapAllStatic",NewEnumEntry="ECC_OverlapAll_Deprecated")
++EnumRedirects=(EnumName="ECollisionChannel",OldEnumEntry="ECC_OverlapAllDynamic_Deprecated",NewEnumEntry="ECC_OverlapAll_Deprecated")
++EnumRedirects=(EnumName="ECollisionChannel",OldEnumEntry="ECC_OverlapAllStatic_Deprecated",NewEnumEntry="ECC_OverlapAll_Deprecated")
 +EnumRedirects=(EnumName="EParticleEventType",OldEnumEntry="EPET_Kismet",NewEnumEntry="EPET_Blueprint")
 
 +EnumRedirects=(EnumName="EKeys",OldEnumEntry="EKeys::XboxTypeS_LeftX",NewEnumEntry="EKeys::Gamepad_LeftX")
diff --git a/Engine/Source/Runtime/Engine/Classes/Engine/EngineTypes.h b/Engine/Source/Runtime/Engine/Classes/Engine/EngineTypes.h
--- a/Engine/Source/Runtime/Engine/Classes/Engine/EngineTypes.h
+++ b/Engine/Source/Runtime/Engine/Classes/Engine/EngineTypes.h
@@ -318,11 +318,7 @@ enum ECollisionChannel
 	 /**
 	  * can't add displaynames because then it will show up in the collision channel option
 	  */
-	ECC_OverlapAll_Deprecated,
-	/** Like ECC_OverlapAll, but only returns objects that are NOT static */
-	ECC_OverlapAllDynamic_Deprecated,
-	/** Like ECC_OverlapAll, but only returns objects that are static */
-	ECC_OverlapAllStatic_Deprecated,
+	ECC_OverlapAll_Deprecated UMETA(Hidden),
 	ECC_MAX,
 };
 
diff --git a/Engine/Source/Runtime/Engine/Private/Collision/PhysXCollision.cpp b/Engine/Source/Runtime/Engine/Private/Collision/PhysXCollision.cpp
--- a/Engine/Source/Runtime/Engine/Private/Collision/PhysXCollision.cpp
+++ b/Engine/Source/Runtime/Engine/Private/Collision/PhysXCollision.cpp
@@ -531,14 +531,6 @@ PxSceneQueryHitType::Enum FPxQueryFilterCallback::CalcQueryHitType(const PxFilte
 		{
 			return PxSceneQueryHitType::eTOUCH;
 		}
-		else if (QuerierChannel == ECC_OverlapAllDynamic_Deprecated)
-		{
-			return bStaticShape ? PxSceneQueryHitType::eNONE : PxSceneQueryHitType::eTOUCH;
-		}
-		else if (QuerierChannel == ECC_OverlapAllStatic_Deprecated)
-		{
-			return !bStaticShape ? PxSceneQueryHitType::eNONE : PxSceneQueryHitType::eTOUCH;
-		}
 		// @todo delete once we fix up object/trace APIs to work separate
 
 		PxU32 const QuerierBit = ECC_TO_BITFIELD(QuerierChannel);
