Forgot the redirect. Sorry.

diff --git a/Engine/Config/BaseEngine.ini b/Engine/Config/BaseEngine.ini
--- a/Engine/Config/BaseEngine.ini
+++ b/Engine/Config/BaseEngine.ini
@@ -439,6 +439,9 @@ bUseStreamingPause=false
 +EnumRedirects=(EnumName="ENavLinkDirection",OldEnumEntry="ESmartNavLinkDir::BothWays",NewEnumEntry="ENavLinkDirection::BothWays")
 +K2FieldRedirects=(OldFieldName="NavigationPath.GetPathLenght",NewFieldName="NavigationPath.GetPathLength")
 
+; mesh component
++TaggedPropertyRedirects=(ClassName="MeshComponent",OldPropertyName="Materials",NewPropertyName="OverrideMaterials")
+
 ;SkeletalMesh redirects
 +TaggedPropertyRedirects=(ClassName="SkeletalMeshLODInfo",OldPropertyName="DisplayFactor",NewPropertyName="ScreenSize")
 
