Fixed HTC One Landscape resolution
From https://github.com/EpicGames/UnrealEngine/pull/603

diff --git a/Engine/Config/BaseEditorUserSettings.ini b/Engine/Config/BaseEditorUserSettings.ini
--- a/Engine/Config/BaseEditorUserSettings.ini
+++ b/Engine/Config/BaseEditorUserSettings.ini
@@ -215,7 +215,7 @@ LastExecutedLaunchPlatform=
 +PhoneScreenResolutions=(Description="Apple iPhone 5, 5S, iTouch 5 (Portrait)",Width=640,Height=1136,AspectRatio="~9:16")
 +PhoneScreenResolutions=(Description="Apple iPhone 5, 5S, iTouch 5 (Landscape)",Width=1136,Height=640,AspectRatio="~16:9")
 +PhoneScreenResolutions=(Description="HTC One (Portrait)",Width=1080,Height=1920,AspectRatio="9:16")
-+PhoneScreenResolutions=(Description="HTC One (Landscape)",Width=1029,Height=1080,AspectRatio="16:9")
++PhoneScreenResolutions=(Description="HTC One (Landscape)",Width=1920,Height=1080,AspectRatio="16:9")
 +PhoneScreenResolutions=(Description="Samsung Galaxy S4 (Portrait)",Width=1080,Height=1920,AspectRatio="9:16")
 +PhoneScreenResolutions=(Description="Samsung Galaxy S4 (Landscape)",Width=1920,Height=1080,AspectRatio="16:9")
 
