no quality loss for bloom in epic setting

diff --git a/Engine/Config/BaseScalability.ini b/Engine/Config/BaseScalability.ini
--- a/Engine/Config/BaseScalability.ini
+++ b/Engine/Config/BaseScalability.ini
@@ -154,7 +154,7 @@ r.LensFlareQuality=2
 r.SceneColorFringeQuality=1
 r.EyeAdaptationQuality=2
 r.BloomQuality=5
-r.FastBlurThreshold=7
+r.FastBlurThreshold=100
 r.Upscale.Quality=3
 r.TonemapperQuality=1
 r.LightShaftQuality=1
