A workaround for a bug in enum property serialization to ini that results in ANavigationData.RuntimeGeneration not being settable to "Static" via Project Settings #UE4

I'll follow up with a Jira ticket to report the issue.

diff --git a/Engine/Config/BaseEngine.ini b/Engine/Config/BaseEngine.ini
--- a/Engine/Config/BaseEngine.ini
+++ b/Engine/Config/BaseEngine.ini
@@ -1947,6 +1947,9 @@ EditorViewButtonsObject=EditorCylinder
 bAutoCreateNavigationData=true
 bAddPlayersToGenerationSeeds=true
 
+[/Script/Engine.NavigationData]
+RuntimeGeneration=Static
+
 [/Script/Engine.RecastNavMesh]
 ; runtime params
 bRebuildAtRuntime=false
