Blueprint Inheritable Components - Enabled

diff --git a/Engine/Config/BaseEngine.ini b/Engine/Config/BaseEngine.ini
--- a/Engine/Config/BaseEngine.ini
+++ b/Engine/Config/BaseEngine.ini
@@ -1483,7 +1483,7 @@ bDisplaysLayout=false
 bNativeCodeGenerationTool=false
 bOptimizeExecutionFlowStack=true
 bOptimizeAdjacentStates=true
-bEnableInheritableComponents=false
+bEnableInheritableComponents=true
 bDeferDependencyLoads=false
 bDeferCDOLoadSerialization=true
 bPropagateDeferredLoadsForStructs=true
