CrashReport - DaysToDeleteUnusedFilesFromPDBCache increased from 7 to 21 days

diff --git a/Engine/Config/BaseEngine.ini b/Engine/Config/BaseEngine.ini
--- a/Engine/Config/BaseEngine.ini
+++ b/Engine/Config/BaseEngine.ini
@@ -1997,7 +1997,7 @@ ExecutablePathPattern=\\epicgames.net\root\Builds\Rocket\Automated\%PRODUCT_VERS
 ; devweb-02
 PDBCachePath=F:/CrashReportPDBCache/
 DepotRoot=F:/depot
-DaysToDeleteUnusedFilesFromPDBCache=7
+DaysToDeleteUnusedFilesFromPDBCache=21
 PDBCacheSizeGB=128
 MinDiskFreeSpaceGB=256
 
