Add defaults for MipGenSettings to all texture groups

diff --git a/Engine/Config/BaseEngine.ini b/Engine/Config/BaseEngine.ini
--- a/Engine/Config/BaseEngine.ini
+++ b/Engine/Config/BaseEngine.ini
@@ -1619,31 +1619,31 @@ FLightPrimitiveInteractionInitialBlockSize=512
 [SystemSettings]
 ; NOTE THAT ANY ITEMS IN THIS SECTION WILL AFFECT ALL PLATFORMS!!!
 
-TEXTUREGROUP_World=(MinLODSize=1,MaxLODSize=4096,LODBias=0,MinMagFilter=aniso,MipFilter=point)
-TEXTUREGROUP_WorldNormalMap=(MinLODSize=1,MaxLODSize=4096,LODBias=0,MinMagFilter=aniso,MipFilter=point)
-TEXTUREGROUP_WorldSpecular=(MinLODSize=1,MaxLODSize=4096,LODBias=0,MinMagFilter=aniso,MipFilter=point)
-TEXTUREGROUP_Character=(MinLODSize=1,MaxLODSize=4096,LODBias=0,MinMagFilter=aniso,MipFilter=point)
-TEXTUREGROUP_CharacterNormalMap=(MinLODSize=1,MaxLODSize=4096,LODBias=0,MinMagFilter=aniso,MipFilter=point)
-TEXTUREGROUP_CharacterSpecular=(MinLODSize=1,MaxLODSize=4096,LODBias=0,MinMagFilter=aniso,MipFilter=point)
-TEXTUREGROUP_Weapon=(MinLODSize=1,MaxLODSize=4096,LODBias=0,MinMagFilter=aniso,MipFilter=point)
-TEXTUREGROUP_WeaponNormalMap=(MinLODSize=1,MaxLODSize=4096,LODBias=0,MinMagFilter=aniso,MipFilter=point)
-TEXTUREGROUP_WeaponSpecular=(MinLODSize=1,MaxLODSize=4096,LODBias=0,MinMagFilter=aniso,MipFilter=point)
-TEXTUREGROUP_Vehicle=(MinLODSize=1,MaxLODSize=4096,LODBias=0,MinMagFilter=aniso,MipFilter=point)
-TEXTUREGROUP_VehicleNormalMap=(MinLODSize=1,MaxLODSize=4096,LODBias=0,MinMagFilter=aniso,MipFilter=point)
-TEXTUREGROUP_VehicleSpecular=(MinLODSize=1,MaxLODSize=4096,LODBias=0,MinMagFilter=aniso,MipFilter=point)
-TEXTUREGROUP_Cinematic=(MinLODSize=1,MaxLODSize=4096,LODBias=0,MinMagFilter=aniso,MipFilter=point)
-TEXTUREGROUP_Effects=(MinLODSize=1,MaxLODSize=4096,LODBias=0,MinMagFilter=linear,MipFilter=point)
-TEXTUREGROUP_EffectsNotFiltered=(MinLODSize=1,MaxLODSize=4096,LODBias=0,MinMagFilter=aniso,MipFilter=point)
-TEXTUREGROUP_Skybox=(MinLODSize=1,MaxLODSize=4096,LODBias=0,MinMagFilter=aniso,MipFilter=point)
-TEXTUREGROUP_UI=(MinLODSize=1,MaxLODSize=4096,LODBias=0,MinMagFilter=aniso,MipFilter=point)
-TEXTUREGROUP_Lightmap=(MinLODSize=1,MaxLODSize=4096,LODBias=0,MinMagFilter=aniso,MipFilter=point)
-TEXTUREGROUP_Shadowmap=(MinLODSize=1,MaxLODSize=4096,LODBias=0,MinMagFilter=aniso,MipFilter=point,NumStreamedMips=3)
-TEXTUREGROUP_RenderTarget=(MinLODSize=1,MaxLODSize=4096,LODBias=0,MinMagFilter=aniso,MipFilter=point)
-TEXTUREGROUP_MobileFlattened=(MinLODSize=1,MaxLODSize=4096,LODBias=0,MinMagFilter=aniso,MipFilter=point)
-TEXTUREGROUP_Terrain_Heightmap=(MinLODSize=1,MaxLODSize=4096,LODBias=0,MinMagFilter=aniso,MipFilter=point)
-TEXTUREGROUP_Terrain_Weightmap=(MinLODSize=1,MaxLODSize=4096,LODBias=0,MinMagFilter=aniso,MipFilter=point)
-TEXTUREGROUP_Bokeh=(MinLODSize=1,MaxLODSize=256,LODBias=0,MinMagFilter=linear,MipFilter=linear)
-TEXTUREGROUP_Pixels2D=(MinLODSize=1,MaxLODSize=4096,LODBias=0,MinMagFilter=point,MipFilter=point)
+TEXTUREGROUP_World=(MinLODSize=1,MaxLODSize=4096,LODBias=0,MinMagFilter=aniso,MipFilter=point,MipGenSettings=TMGS_SimpleAverage)
+TEXTUREGROUP_WorldNormalMap=(MinLODSize=1,MaxLODSize=4096,LODBias=0,MinMagFilter=aniso,MipFilter=point,MipGenSettings=TMGS_SimpleAverage)
+TEXTUREGROUP_WorldSpecular=(MinLODSize=1,MaxLODSize=4096,LODBias=0,MinMagFilter=aniso,MipFilter=point,MipGenSettings=TMGS_SimpleAverage)
+TEXTUREGROUP_Character=(MinLODSize=1,MaxLODSize=4096,LODBias=0,MinMagFilter=aniso,MipFilter=point,MipGenSettings=TMGS_SimpleAverage)
+TEXTUREGROUP_CharacterNormalMap=(MinLODSize=1,MaxLODSize=4096,LODBias=0,MinMagFilter=aniso,MipFilter=point,MipGenSettings=TMGS_SimpleAverage)
+TEXTUREGROUP_CharacterSpecular=(MinLODSize=1,MaxLODSize=4096,LODBias=0,MinMagFilter=aniso,MipFilter=point,MipGenSettings=TMGS_SimpleAverage)
+TEXTUREGROUP_Weapon=(MinLODSize=1,MaxLODSize=4096,LODBias=0,MinMagFilter=aniso,MipFilter=point,MipGenSettings=TMGS_SimpleAverage)
+TEXTUREGROUP_WeaponNormalMap=(MinLODSize=1,MaxLODSize=4096,LODBias=0,MinMagFilter=aniso,MipFilter=point,MipGenSettings=TMGS_SimpleAverage)
+TEXTUREGROUP_WeaponSpecular=(MinLODSize=1,MaxLODSize=4096,LODBias=0,MinMagFilter=aniso,MipFilter=point,MipGenSettings=TMGS_SimpleAverage)
+TEXTUREGROUP_Vehicle=(MinLODSize=1,MaxLODSize=4096,LODBias=0,MinMagFilter=aniso,MipFilter=point,MipGenSettings=TMGS_SimpleAverage)
+TEXTUREGROUP_VehicleNormalMap=(MinLODSize=1,MaxLODSize=4096,LODBias=0,MinMagFilter=aniso,MipFilter=point,MipGenSettings=TMGS_SimpleAverage)
+TEXTUREGROUP_VehicleSpecular=(MinLODSize=1,MaxLODSize=4096,LODBias=0,MinMagFilter=aniso,MipFilter=point,MipGenSettings=TMGS_SimpleAverage)
+TEXTUREGROUP_Cinematic=(MinLODSize=1,MaxLODSize=4096,LODBias=0,MinMagFilter=aniso,MipFilter=point,MipGenSettings=TMGS_SimpleAverage)
+TEXTUREGROUP_Effects=(MinLODSize=1,MaxLODSize=4096,LODBias=0,MinMagFilter=linear,MipFilter=point,MipGenSettings=TMGS_SimpleAverage)
+TEXTUREGROUP_EffectsNotFiltered=(MinLODSize=1,MaxLODSize=4096,LODBias=0,MinMagFilter=aniso,MipFilter=point,MipGenSettings=TMGS_SimpleAverage)
+TEXTUREGROUP_Skybox=(MinLODSize=1,MaxLODSize=4096,LODBias=0,MinMagFilter=aniso,MipFilter=point,MipGenSettings=TMGS_SimpleAverage)
+TEXTUREGROUP_UI=(MinLODSize=1,MaxLODSize=4096,LODBias=0,MinMagFilter=aniso,MipFilter=point,MipGenSettings=TMGS_SimpleAverage)
+TEXTUREGROUP_Lightmap=(MinLODSize=1,MaxLODSize=4096,LODBias=0,MinMagFilter=aniso,MipFilter=point,MipGenSettings=TMGS_SimpleAverage)
+TEXTUREGROUP_Shadowmap=(MinLODSize=1,MaxLODSize=4096,LODBias=0,MinMagFilter=aniso,MipFilter=point,NumStreamedMips=3,MipGenSettings=TMGS_SimpleAverage)
+TEXTUREGROUP_RenderTarget=(MinLODSize=1,MaxLODSize=4096,LODBias=0,MinMagFilter=aniso,MipFilter=point,MipGenSettings=TMGS_SimpleAverage)
+TEXTUREGROUP_MobileFlattened=(MinLODSize=1,MaxLODSize=4096,LODBias=0,MinMagFilter=aniso,MipFilter=point,MipGenSettings=TMGS_SimpleAverage)
+TEXTUREGROUP_Terrain_Heightmap=(MinLODSize=1,MaxLODSize=4096,LODBias=0,MinMagFilter=aniso,MipFilter=point,MipGenSettings=TMGS_SimpleAverage)
+TEXTUREGROUP_Terrain_Weightmap=(MinLODSize=1,MaxLODSize=4096,LODBias=0,MinMagFilter=aniso,MipFilter=point,MipGenSettings=TMGS_SimpleAverage)
+TEXTUREGROUP_Bokeh=(MinLODSize=1,MaxLODSize=256,LODBias=0,MinMagFilter=linear,MipFilter=linear,MipGenSettings=TMGS_SimpleAverage)
+TEXTUREGROUP_Pixels2D=(MinLODSize=1,MaxLODSize=4096,LODBias=0,MinMagFilter=point,MipFilter=point,MipGenSettings=TMGS_SimpleAverage)
 
 [SystemSettingsEditor]
 ; System settings overrides for the editor.  Ideally the editor should use the same settings as the game.
