Disable parallel rendering in the editor

diff --git a/Engine/Config/BaseEngine.ini b/Engine/Config/BaseEngine.ini
--- a/Engine/Config/BaseEngine.ini
+++ b/Engine/Config/BaseEngine.ini
@@ -1648,6 +1648,9 @@ TEXTUREGROUP_Pixels2D=(MinLODSize=1,MaxLODSize=4096,LODBias=0,MinMagFilter=point
 [SystemSettingsEditor]
 ; System settings overrides for the editor.  Ideally the editor should use the same settings as the game.
 
+; Parallel rendering has not been tested in the editor so just disable it.
+r.RHICmdBypass=1
+
 [SystemSettingsSplitScreen2]
 ; System settings overrides for split screen
 ; Use medium detail mode in splitscreen, this allows LD's to mark meshes as High detail and they won't render in SS
