[INTEGRATE] Change 2423633 by [REDACTED_PERSON][REDACTED_HANDLE] on 2015/01/29 08:00:36

	Remove PakRead node from Rocket DDC graph; since PakWrite node now merges existing data, there's no need to read from it too. Hopefully fixes build failure on Mac (sharing violation on file?)

diff --git a/Engine/Config/BaseEngine.ini b/Engine/Config/BaseEngine.ini
--- a/Engine/Config/BaseEngine.ini
+++ b/Engine/Config/BaseEngine.ini
@@ -1888,7 +1888,6 @@ Boot=(Type=Boot, Filename="%GAMEDIR%DerivedDataCache/Boot.ddc", MaxCacheSize=512
 Local=(Type=FileSystem, ReadOnly=false, Clean=false, Flush=false, PurgeTransient=true, DeleteUnused=true, UnusedFileAge=34, FoldersToClean=-1, Path=../../../Engine/DerivedDataCache, EnvPathOverride=UE-LocalDataCachePath)
 Shared=(Type=FileSystem, ReadOnly=false, Clean=false, Flush=false, DeleteUnused=false, UnusedFileAge=23, FoldersToClean=-1, Path=\\epicgames.net\root\UE4DDC, EnvPathOverride=UE-SharedDataCachePath)
 AltShared=(Type=FileSystem, ReadOnly=true, Clean=false, Flush=false, DeleteUnused=false, UnusedFileAge=23, FoldersToClean=-1, Path=\\epicgames.net\root\UE4DDC2, EnvPathOverride=UE-SharedDataCachePath2)
-PakRead=(Type=ReadPak, Filename=../../../Engine/DerivedDataCache/DDC.ddp)
 PakWrite=(Type=WritePak, Filename=../../../Engine/DerivedDataCache/DDC.ddp, Merge=true)
 
 [CreateProjectCache]
