Merging engine content and templates from 4.7 branch.

diff --git a/Engine/Config/BaseEngine.ini b/Engine/Config/BaseEngine.ini
--- a/Engine/Config/BaseEngine.ini
+++ b/Engine/Config/BaseEngine.ini
@@ -1481,6 +1481,7 @@ bReinstanceOnlyWhenNecessary=true
 bChangeDefaultValueWithoutReinstancing=true
 bDisplaysLayout=false
 bNativeCodeGenerationTool=false
+bSkeletonInheritSkeletonClasses=false
 bOptimizeExecutionFlowStack=true
 bOptimizeAdjacentStates=true
 bEnableInheritableComponents=true
