nicer autocomplete

diff --git a/Engine/Config/BaseInput.ini b/Engine/Config/BaseInput.ini
--- a/Engine/Config/BaseInput.ini
+++ b/Engine/Config/BaseInput.ini
@@ -51,6 +51,9 @@ MaxScrollbackSize=1024
 +AutoCompleteMapPaths=Content/Maps
 +ManualAutoCompleteList=(Command="Exit",Desc="Exit (Exits the game)")
 +ManualAutoCompleteList=(Command="DebugCreatePlayer 1",Desc=)
++ManualAutoCompleteList=(Command="Shot",Desc="Make a screenshot")
++ManualAutoCompleteList=(Command="RecompileShaders changed",Desc="Recompile the shaders that have changed")
++ManualAutoCompleteList=(Command="RecompileShaders global",Desc="Recompile all global shaders")
 +ManualAutoCompleteList=(Command="StartFPSChart",Desc="StartFPSChart (after that use StopFPSChart)")
 +ManualAutoCompleteList=(Command="StopFPSChart",Desc="StopFPSChart (after that look for the output in Saved/Profiling/FPSChartStats)")
 +ManualAutoCompleteList=(Command="FreezeAt",Desc="Locks the player view and rendering time.")
