[UE-8798] Data for tutorial-browser-ordering issue.

diff --git a/Engine/Config/BaseEditorGameAgnostic.ini b/Engine/Config/BaseEditorGameAgnostic.ini
--- a/Engine/Config/BaseEditorGameAgnostic.ini
+++ b/Engine/Config/BaseEditorGameAgnostic.ini
@@ -18,14 +18,15 @@ bCaptureAudioFromComputer=true
 bCaptureAudioFromMicrophone=true
 
 [/Script/IntroTutorials.EditorTutorialSettings]
-+Categories=(Identifier="Basics",Title="Basics",Description="Getting started with Unreal Engine 4.",Icon="PlayWorld.RepeatLastPlay",Texture=None)
-+Categories=(Identifier="Mobile",Title="Mobile",Description="Mobile Tutorials.",Icon="MaterialEditor.ToggleMobileStats",Texture=None)
-+Categories=(Identifier="Blueprints",Title="Blueprints",Description="Tutorials covering the creation and usage of Blueprints.",Icon="FullBlueprintEditor.EditGlobalOptions",Texture=/Engine/Tutorial/BlueprintTutorials/TutorialAssets/Blueprint_64x.Blueprint_64x)
-+Categories=(Identifier="Paper2D",Title="Paper2D",Description="Tutorials convering the usage of Unreal Engine 4\'s 2D game tools: Paper2D.",Icon="ClassThumbnail.PaperSprite",Texture=/Engine/Tutorial/Paper2D/TutorialAssets/Paper2DSprite_TutorialIcon.Paper2DSprite_TutorialIcon)
-+Categories=(Identifier="Landscape",Title="Landscape",Description="Tutorials covering the Unreal Editor 4 terrain editor: Landscape.",Icon="LevelEditor.LandscapeMode",Texture=/Engine/Tutorial/Landscape/TutorialAssets/Landscape.Landscape)
-+Categories=(Identifier="Foliage",Title="Foliage",Description="Tutorials covering the Unreal Engine 4 Foliage tool.",Icon="LevelEditor.FoliageMode",Texture=/Engine/Tutorial/Foliage/TutorialAssets/Foliage.Foliage)
-+Categories=(Identifier="Editors",Title="Editors",Description="Tutorials covering the various sub editors of Unreal Engine 4.",Icon=,Texture=/Engine/Tutorial/SubEditors/TutorialAssets/icon_Editor_Preferences_General_40x.icon_Editor_Preferences_General_40x)
-+Categories=(Identifier="Animation",Title="Animation",Description="Tutorials convering the animation system in Unreal Engine 4.",Icon=,Texture=/Engine/Tutorial/SubEditors/TutorialAssets/icon_ShowSkeletalMeshes_40x.icon_ShowSkeletalMeshes_40x)
++Categories=(Identifier="Basics",Title="Basics",Description="Getting started with Unreal Engine 4.",Icon="PlayWorld.RepeatLastPlay",Texture=None,SortOrder=100)
++Categories=(Identifier="Mobile",Title="Mobile",Description="Mobile Tutorials.",Icon="MaterialEditor.ToggleMobileStats",Texture=None,SortOrder=200)
++Categories=(Identifier="Blueprints",Title="Blueprints",Description="Tutorials covering the creation and usage of Blueprints.",Icon="FullBlueprintEditor.EditGlobalOptions",Texture=/Engine/Tutorial/BlueprintTutorials/TutorialAssets/Blueprint_64x.Blueprint_64x,SortOrder=300)
++Categories=(Identifier="Paper2D",Title="Paper2D",Description="Tutorials convering the usage of Unreal Engine 4\'s 2D game tools: Paper2D.",Icon="ClassThumbnail.PaperSprite",Texture=/Engine/Tutorial/Paper2D/TutorialAssets/Paper2DSprite_TutorialIcon.Paper2DSprite_TutorialIcon,SortOrder=400)
++Categories=(Identifier="Landscape",Title="Landscape",Description="Tutorials covering the Unreal Editor 4 terrain editor: Landscape.",Icon="LevelEditor.LandscapeMode",Texture=/Engine/Tutorial/Landscape/TutorialAssets/Landscape.Landscape,SortOrder=500)
++Categories=(Identifier="Foliage",Title="Foliage",Description="Tutorials covering the Unreal Engine 4 Foliage tool.",Icon="LevelEditor.FoliageMode",Texture=/Engine/Tutorial/Foliage/TutorialAssets/Foliage.Foliage,SortOrder=600)
++Categories=(Identifier="Editors",Title="Editors",Description="Tutorials covering the various sub editors of Unreal Engine 4.",Icon=,Texture=/Engine/Tutorial/SubEditors/TutorialAssets/icon_Editor_Preferences_General_40x.icon_Editor_Preferences_General_40x,SortOrder=700)
++Categories=(Identifier="Animation",Title="Animation",Description="Tutorials convering the animation system in Unreal Engine 4.",Icon=,Texture=/Engine/Tutorial/SubEditors/TutorialAssets/icon_ShowSkeletalMeshes_40x.icon_ShowSkeletalMeshes_40x,SortOrder=800)
++Categories=(Identifier="Code",Title="Code",Description="Write C++ code for use in Unreal Engine 4.",Icon="PlayWorld.RepeatLastPlay",Texture=None,SortOrder=900)
 +StartupTutorial=/Engine/Tutorial/Basics/LevelEditorAttract.LevelEditorAttract_C
 +TutorialContexts=(Context="StaticMeshEditor",BrowserFilter=,AttractTutorial=None,LaunchTutorial=/Engine/Tutorial/SubEditors/StaticMeshEditorTutorial.StaticMeshEditorTutorial_C)
 +TutorialContexts=(Context="LevelEditor",BrowserFilter=,AttractTutorial=None,LaunchTutorial=/Engine/Tutorial/Basics/LevelEditorOverview.LevelEditorOverview_C)
