Remove unused properties

diff --git a/Engine/Config/BaseEngine.ini b/Engine/Config/BaseEngine.ini
--- a/Engine/Config/BaseEngine.ini
+++ b/Engine/Config/BaseEngine.ini
@@ -1277,7 +1277,6 @@ AllowPeerConnections=False
 AllowPeerVoice=False
 ConnectionTimeout=60.0
 InitialConnectTimeout=60.0
-AckTimeout=1.0
 KeepAliveTime=0.2
 MaxClientRate=15000
 MaxInternetClientRate=10000
