Removing more unused properties

diff --git a/Engine/Config/BaseEngine.ini b/Engine/Config/BaseEngine.ini
--- a/Engine/Config/BaseEngine.ini
+++ b/Engine/Config/BaseEngine.ini
@@ -1284,7 +1284,6 @@ RelevantTimeout=5.0
 SpawnPrioritySeconds=1.0
 ServerTravelPause=4.0
 NetServerMaxTickRate=30
-LanServerMaxTickRate=35
 NetConnectionClassName="/Script/OnlineSubsystemUtils.IpConnection"
 MaxPortCountToTry=512
 
