Set PackagingSettings.FullRebuild to be false by default in config file.

diff --git a/Engine/Config/BaseGame.ini b/Engine/Config/BaseGame.ini
--- a/Engine/Config/BaseGame.ini
+++ b/Engine/Config/BaseGame.ini
@@ -34,7 +34,7 @@ SupportContact=
 
 [/Script/UnrealEd.ProjectPackagingSettings]
 BuildConfiguration=PPBC_Development
-FullRebuild=True
+FullRebuild=False
 UsePakFile=True
 bGenerateChunks=False
 IncludePrerequisites=True
