Set keep hierarchy default to be true

diff --git a/Engine/Config/BaseEditorUserSettings.ini b/Engine/Config/BaseEditorUserSettings.ini
--- a/Engine/Config/BaseEditorUserSettings.ini
+++ b/Engine/Config/BaseEditorUserSettings.ini
@@ -53,6 +53,9 @@ PropertyMatrix_NumberOfPasteOperationsBeforeWarning=20
 ; Blueprint Editor SCS Editor settings
 bSCSEditorShowGrid=true
 
+; Export
+bKeepAttachHierarchy=true
+
 ; Misc
 HoverHighlightIntensity=0.000000
 bDisplayUIExtensionPoints=False
