Fixed the "Common" category in the UMG palette view not open by default.  User expansion is still saved

diff --git a/Engine/Config/BaseEditorUserSettings.ini b/Engine/Config/BaseEditorUserSettings.ini
--- a/Engine/Config/BaseEditorUserSettings.ini
+++ b/Engine/Config/BaseEditorUserSettings.ini
@@ -686,3 +686,6 @@ DefaultBrushRadius=128
 +Node=(Class=MaterialExpressionConstant2Vector Key=Two Shift=false Ctrl=false Alt=false)
 +Node=(Class=MaterialExpressionConstant3Vector Key=Three Shift=false Ctrl=false Alt=false)
 +Node=(Class=MaterialExpressionConstant4Vector Key=Four Shift=false Ctrl=false Alt=false)
+
+[WidgetTemplatesExpanded]
+Common=True
\ No newline at end of file
