[INTEGRATE] Content from 4.7 release branch

diff --git a/Engine/Build/JunkManifest.txt b/Engine/Build/JunkManifest.txt
--- a/Engine/Build/JunkManifest.txt
+++ b/Engine/Build/JunkManifest.txt
@@ -51,5 +51,6 @@ Engine\Plugins\Script\ScriptGeneratorPlugin\Binaries\Win64\*
 Engine\Plugins\Script\ScriptPlugin\Binaries\Mac\*
 Engine\Plugins\Script\ScriptPlugin\Binaries\Win32\*
 Engine\Plugins\Script\ScriptPlugin\Binaries\Win64\*
+Engine\Plugins\Runtime\WebBrowserAsPlugin\Binaries
 Engine\Plugins\Experimental\Paper2D\Binaries
 Engine\Plugins\Experimental\Paper2D\Intermediate
diff --git a/Engine/Config/BaseEditor.ini b/Engine/Config/BaseEditor.ini
--- a/Engine/Config/BaseEditor.ini
+++ b/Engine/Config/BaseEditor.ini
@@ -243,7 +243,7 @@ DocumentationURL="http://udn.unrealengine.com/docs/ue4/INT/index.html"
 AskAQuestionURL="http://udn.unrealengine.com/questions/ask.html"
 SearchForAnswersURL="http://udn.unrealengine.com/search.html"
 SupportWebsiteURL="http://udn.unrealengine.com/index.html"
-SourceCodeIDEURL_Windows="http://go.microsoft.com/?linkid=9863608"
+SourceCodeIDEURL_Windows="http://go.microsoft.com/fwlink/?LinkId=525621"
 SourceCodeIDEURL_Mac="https://developer.apple.com/xcode"
 SourceCodeIDEURL_Other="http://en.wikipedia.org/wiki/Integrated_development_environment"
 EpicGamesURL="http://UnrealEngine.com"
diff --git a/Engine/Config/BaseEditorGameAgnostic.ini b/Engine/Config/BaseEditorGameAgnostic.ini
--- a/Engine/Config/BaseEditorGameAgnostic.ini
+++ b/Engine/Config/BaseEditorGameAgnostic.ini
@@ -1,5 +1,5 @@
 [/Script/UnrealEd.EditorGameAgnosticSettings]
-bCopyStarterContentPreference = false
+bCopyStarterContentPreference = true
 
 [/Script/NewsFeed.NewsFeedSettings]
 CdnSourceUrl="http://cdn.unrealengine.com/newsfeed/[[Culture]]/index.json"
@@ -20,14 +20,14 @@ bCaptureAudioFromMicrophone=true
 
 [/Script/IntroTutorials.EditorTutorialSettings]
 +Categories=(Identifier="Basics",Title="Basics",Description="Getting started with Unreal Engine 4.",Icon="PlayWorld.RepeatLastPlay",Texture=None,SortOrder=100)
-+Categories=(Identifier="Mobile",Title="Mobile",Description="Mobile Tutorials.",Icon="MaterialEditor.ToggleMobileStats",Texture=None,SortOrder=200)
-+Categories=(Identifier="Blueprints",Title="Blueprints",Description="Tutorials covering the creation and usage of Blueprints.",Icon="FullBlueprintEditor.EditGlobalOptions",Texture=/Engine/Tutorial/BlueprintTutorials/TutorialAssets/Blueprint_64x.Blueprint_64x,SortOrder=300)
-+Categories=(Identifier="Paper2D",Title="Paper2D",Description="Tutorials convering the usage of Unreal Engine 4\'s 2D game tools: Paper2D.",Icon="ClassThumbnail.PaperSprite",Texture=/Engine/Tutorial/Paper2D/TutorialAssets/Paper2DSprite_TutorialIcon.Paper2DSprite_TutorialIcon,SortOrder=400)
-+Categories=(Identifier="Landscape",Title="Landscape",Description="Tutorials covering the Unreal Editor 4 terrain editor: Landscape.",Icon="LevelEditor.LandscapeMode",Texture=/Engine/Tutorial/Landscape/TutorialAssets/Landscape.Landscape,SortOrder=500)
-+Categories=(Identifier="Foliage",Title="Foliage",Description="Tutorials covering the Unreal Engine 4 Foliage tool.",Icon="LevelEditor.FoliageMode",Texture=/Engine/Tutorial/Foliage/TutorialAssets/Foliage.Foliage,SortOrder=600)
-+Categories=(Identifier="Editors",Title="Editors",Description="Tutorials covering the various sub editors of Unreal Engine 4.",Icon=,Texture=/Engine/Tutorial/SubEditors/TutorialAssets/icon_Editor_Preferences_General_40x.icon_Editor_Preferences_General_40x,SortOrder=700)
-+Categories=(Identifier="Animation",Title="Animation",Description="Tutorials convering the animation system in Unreal Engine 4.",Icon=,Texture=/Engine/Tutorial/SubEditors/TutorialAssets/icon_ShowSkeletalMeshes_40x.icon_ShowSkeletalMeshes_40x,SortOrder=800)
-+Categories=(Identifier="Code",Title="Code",Description="Write C++ code for use in Unreal Engine 4.",Icon="PlayWorld.RepeatLastPlay",Texture=None,SortOrder=900)
++Categories=(Identifier="Blueprints",Title="Blueprints",Description="Tutorials covering the creation and usage of Blueprints.",Icon="FullBlueprintEditor.EditGlobalOptions",Texture=/Engine/Tutorial/BlueprintTutorials/TutorialAssets/Blueprint_64x.Blueprint_64x,SortOrder=200)
++Categories=(Identifier="Editors",Title="Editors",Description="Tutorials covering the various sub editors of Unreal Engine 4.",Icon=,Texture=/Engine/Tutorial/SubEditors/TutorialAssets/icon_Editor_Preferences_General_40x.icon_Editor_Preferences_General_40x,SortOrder=225)
++Categories=(Identifier="Code",Title="Code",Description="Write C++ code for use in Unreal Engine 4.",Icon="PlayWorld.RepeatLastPlay",Texture=None,SortOrder=250)
++Categories=(Identifier="Paper2D",Title="Paper2D",Description="Tutorials convering the usage of Unreal Engine 4\'s 2D game tools: Paper2D.",Icon="ClassThumbnail.PaperSprite",Texture=/Engine/Tutorial/Paper2D/TutorialAssets/Paper2DSprite_TutorialIcon.Paper2DSprite_TutorialIcon,SortOrder=500)
++Categories=(Identifier="Animation",Title="Animation",Description="Tutorials convering the animation system in Unreal Engine 4.",Icon=,Texture=/Engine/Tutorial/SubEditors/TutorialAssets/icon_ShowSkeletalMeshes_40x.icon_ShowSkeletalMeshes_40x,SortOrder=600)
++Categories=(Identifier="Landscape",Title="Landscape",Description="Tutorials covering the Unreal Editor 4 terrain editor: Landscape.",Icon="LevelEditor.LandscapeMode",Texture=/Engine/Tutorial/Landscape/TutorialAssets/Landscape.Landscape,SortOrder=700)
++Categories=(Identifier="Foliage",Title="Foliage",Description="Tutorials covering the Unreal Engine 4 Foliage tool.",Icon="LevelEditor.FoliageMode",Texture=/Engine/Tutorial/Foliage/TutorialAssets/Foliage.Foliage,SortOrder=800)
++Categories=(Identifier="Mobile",Title="Mobile",Description="Mobile Tutorials.",Icon="MaterialEditor.ToggleMobileStats",Texture=None,SortOrder=900)
 +StartupTutorial=/Engine/Tutorial/Basics/LevelEditorAttract.LevelEditorAttract_C
 +TutorialContexts=(Context="StaticMeshEditor",BrowserFilter=,AttractTutorial=None,LaunchTutorial=/Engine/Tutorial/SubEditors/StaticMeshEditorTutorial.StaticMeshEditorTutorial_C)
 +TutorialContexts=(Context="LevelEditor",BrowserFilter=,AttractTutorial=None,LaunchTutorial=/Engine/Tutorial/Basics/LevelEditorOverview.LevelEditorOverview_C)
