Added all show flags to the auto-complete list.

https://github.com/EpicGames/UnrealEngine/pull/787

diff --git a/Engine/Config/BaseInput.ini b/Engine/Config/BaseInput.ini
--- a/Engine/Config/BaseInput.ini
+++ b/Engine/Config/BaseInput.ini
@@ -76,16 +76,6 @@ MaxScrollbackSize=1024
 +ManualAutoCompleteList=(Command="ReloadLoc ",Desc="ReloadLoc <Class/ObjectName> (Reloads localized variables for the specified object/class)")
 +ManualAutoCompleteList=(Command="Set ",Desc="Set <ClassName> <PropertyName> <Value> (Sets property to value on objectname)")
 +ManualAutoCompleteList=(Command="SetNoPEC",Desc="SetNoPEC <ClassName> <PropertyName> <Value> (Sets property to value on objectname without Pre/Post Edit Change notifications)")
-+ManualAutoCompleteList=(Command="Show BOUNDS",Desc="Show BOUNDS (Displays bounding boxes for all visible objects)")
-+ManualAutoCompleteList=(Command="Show BSP",Desc="Show BSP (Toggles BSP rendering)")
-+ManualAutoCompleteList=(Command="Show COLLISION",Desc="Show COLLISION (Toggles collision rendering)")
-+ManualAutoCompleteList=(Command="Show FOG",Desc="Show FOG (Toggles fog rendering)")
-+ManualAutoCompleteList=(Command="Show LEVELCOLORATION",Desc="Show LEVELCOLORATION (Toggles per-level coloration)")
-+ManualAutoCompleteList=(Command="Show NAVIGATION",Desc="Show NAVIGATION (Toggles navigation rendering)")
-+ManualAutoCompleteList=(Command="Show PostProcessing",Desc="Show PostProcessing (Toggles post process rendering)")
-+ManualAutoCompleteList=(Command="Show SKELETALMESHES",Desc="Show SKELETALMESHES (Toggles skeletal mesh rendering)")
-+ManualAutoCompleteList=(Command="Show TERRAIN",Desc="Show TERRAIN (Toggles terrain rendering)")
-+ManualAutoCompleteList=(Command="Show SPLINES",Desc="Show SPLINES (Toggles spline rendering)")
 +ManualAutoCompleteList=(Command="Stat FPS",Desc="Stat FPS (Shows FPS counter)")
 +ManualAutoCompleteList=(Command="Stat UNIT",Desc="Stat UNIT (Shows hardware unit framerate)")
 +ManualAutoCompleteList=(Command="Stat UnitGraph",Desc="Stat UnitGraph (Draws simple unit time graph)")
diff --git a/Engine/Source/Runtime/Engine/Private/UserInterface/Console.cpp b/Engine/Source/Runtime/Engine/Private/UserInterface/Console.cpp
--- a/Engine/Source/Runtime/Engine/Private/UserInterface/Console.cpp
+++ b/Engine/Source/Runtime/Engine/Private/UserInterface/Console.cpp
@@ -239,6 +239,35 @@ void UConsole::BuildRuntimeAutoCompleteList(bool bForce)
 	}
 #endif
 
+	// Add all showflag commands.
+	{
+		struct FIterSink
+		{
+			FIterSink(TArray<FAutoCompleteCommand>& InAutoCompleteList)
+				: AutoCompleteList(InAutoCompleteList)
+			{
+			}
+			
+			bool OnEngineShowFlag(uint32 InIndex, const FString& InName)
+			{
+				// Get localized name.
+				FText LocName;
+				FEngineShowFlags::FindShowFlagDisplayName(InName, LocName);
+				
+				int32 NewIdx = AutoCompleteList.AddZeroed(1);
+				AutoCompleteList[NewIdx].Command = TEXT("show ") + InName;
+				AutoCompleteList[NewIdx].Desc = FString::Printf(TEXT("show %s (toggles the %s showflag)"),*InName, *LocName.ToString());
+				
+				return true;
+			}
+			
+			TArray<FAutoCompleteCommand>& AutoCompleteList;
+		};
+		
+		FIterSink Sink(AutoCompleteList);
+		FEngineShowFlags::IterateAllFlags(Sink);
+	}
+
 	// build the magic tree!
 	for (int32 ListIdx = 0; ListIdx < AutoCompleteList.Num(); ListIdx++)
 	{
