Core - Misc config clean-up

diff --git a/Engine/Config/BaseEngine.ini b/Engine/Config/BaseEngine.ini
--- a/Engine/Config/BaseEngine.ini
+++ b/Engine/Config/BaseEngine.ini
@@ -2000,21 +2000,6 @@ Timeout=1.0f
 +IniToLoad=EditorUserSettings
 +IniToLoad=EditorKeyBindings
 
-[Stats]
-AlwaysRecordingStats=true
-
-[Engine.CrashDebugHelper]
-; ini files do not like '//' - they treat it as a comment
-LocalSymbolStore=d:/Symbols/UE4
-SourceControlBuildLabelPattern=*-CL-%CHANGELISTNUMBER%
-ExecutablePathPattern=\\epicgames.net\root\Builds\Rocket\Automated\%PRODUCT_VERSION%\Windows
-; devweb-02
-PDBCachePath=F:/CrashReportPDBCache/
-DepotRoot=F:/depot
-DaysToDeleteUnusedFilesFromPDBCache=21
-PDBCacheSizeGB=128
-MinDiskFreeSpaceGB=256
-
 [Engine.ErrorHandling]
 bPromptForRemoteDebugging=false
 bPromptForRemoteDebugOnEnsure=false
