#ue4 - Fix improperly tagged redirects

diff --git a/Engine/Config/BaseEngine.ini b/Engine/Config/BaseEngine.ini
--- a/Engine/Config/BaseEngine.ini
+++ b/Engine/Config/BaseEngine.ini
@@ -894,25 +894,25 @@ bUseStreamingPause=false
 +K2FieldRedirects=(OldFieldName="DefaultPawn.DefaultPawnMovement",NewFieldName="DefaultPawn.MovementComponent")
 +K2FieldRedirects=(OldFieldName="FloatingPawnMovement.AccelRate",NewFieldName="FloatingPawnMovement.Acceleration")
 +K2FieldRedirects=(OldFieldName="FloatingPawnMovement.DecelRate",NewFieldName="FloatingPawnMovement.Deceleration")
-+K2FieldRedirects=(OldFieldName="ProjectileMovementComponent.Speed",NewPropertyName="ProjectileMovementComponent.InitialSpeed")
++K2FieldRedirects=(OldFieldName="ProjectileMovementComponent.Speed",NewFieldName="ProjectileMovementComponent.InitialSpeed")
 +K2FieldRedirects=(OldFieldName="CharacterMovementComponent.AccelRate",NewFieldName="CharacterMovementComponent.MaxAcceleration")
 +K2FieldRedirects=(OldFieldName="CharacterMovementComponent.CrouchedPct",NewFieldName="CharacterMovementComponent.CrouchedSpeedMultiplier")
 +K2FieldRedirects=(OldFieldName="CharacterMovementComponent.CrouchHeight",NewFieldName="CharacterMovementComponent.CrouchedHalfHeight")
-+K2FieldRedirects=(OldFieldName="CharacterMovementComponent.JumpZ",NewPropertyName="CharacterMovementComponent.JumpZVelocity")
-+K2FieldRedirects=(OldFieldName="CharacterMovementComponent.GroundSpeed",NewPropertyName="CharacterMovementComponent.MaxWalkSpeed")
-+K2FieldRedirects=(OldFieldName="CharacterMovementComponent.AirSpeed",NewPropertyName="CharacterMovementComponent.MaxFlySpeed")
-+K2FieldRedirects=(OldFieldName="CharacterMovementComponent.WaterSpeed",NewPropertyName="CharacterMovementComponent.MaxSwimSpeed")
-+K2FieldRedirects=(OldFieldName="CharacterMovementComponent.CrouchedSpeedPercent",NewPropertyName="CharacterMovementComponent.CrouchedSpeedMultiplier")
-+K2FieldRedirects=(OldFieldName="CharacterMovementComponent.bOrientToMovement",NewPropertyName="CharacterMovementComponent.bOrientRotationToMovement")
-+K2FieldRedirects=(OldFieldName="CharacterMovementComponent.bCrouchMovesCharacterDown",NewPropertyName="CharacterMovementComponent.bCrouchMaintainsBaseLocation")
++K2FieldRedirects=(OldFieldName="CharacterMovementComponent.JumpZ",NewFieldName="CharacterMovementComponent.JumpZVelocity")
++K2FieldRedirects=(OldFieldName="CharacterMovementComponent.GroundSpeed",NewFieldName="CharacterMovementComponent.MaxWalkSpeed")
++K2FieldRedirects=(OldFieldName="CharacterMovementComponent.AirSpeed",NewFieldName="CharacterMovementComponent.MaxFlySpeed")
++K2FieldRedirects=(OldFieldName="CharacterMovementComponent.WaterSpeed",NewFieldName="CharacterMovementComponent.MaxSwimSpeed")
++K2FieldRedirects=(OldFieldName="CharacterMovementComponent.CrouchedSpeedPercent",NewFieldName="CharacterMovementComponent.CrouchedSpeedMultiplier")
++K2FieldRedirects=(OldFieldName="CharacterMovementComponent.bOrientToMovement",NewFieldName="CharacterMovementComponent.bOrientRotationToMovement")
++K2FieldRedirects=(OldFieldName="CharacterMovementComponent.bCrouchMovesCharacterDown",NewFieldName="CharacterMovementComponent.bCrouchMaintainsBaseLocation")
 +K2FieldRedirects=(OldFieldName="CharacterMovementComponent.AddMomentum",NewFieldName="CharacterMovementComponent.AddImpulse")
 +K2ParamRedirects=(NodeName="/Script/BlueprintGraph.K2Node_CallFunction", OldParamName="CharacterMovementComponent.AddImpulse.InMomentum", NewParamName="Impulse")
 +K2ParamRedirects=(NodeName="/Script/BlueprintGraph.K2Node_CallFunction", OldParamName="CharacterMovementComponent.AddImpulse.bMassIndependent", NewParamName="bVelocityChange")
 
-+K2FieldRedirects=(OldFieldName="CameraComponent.bUseControllerViewRotation",NewPropertyName="CameraComponent.bUsePawnControlRotation")
-+K2FieldRedirects=(OldFieldName="CameraComponent.bUsePawnViewRotation",NewPropertyName="CameraComponent.bUsePawnControlRotation")
-+K2FieldRedirects=(OldFieldName="SpringArmComponent.bUseControllerViewRotation",NewPropertyName="SpringArmComponent.bUsePawnControlRotation")
-+K2FieldRedirects=(OldFieldName="SpringArmComponent.bUsePawnViewRotation",NewPropertyName="SpringArmComponent.bUsePawnControlRotation")
++K2FieldRedirects=(OldFieldName="CameraComponent.bUseControllerViewRotation",NewFieldName="CameraComponent.bUsePawnControlRotation")
++K2FieldRedirects=(OldFieldName="CameraComponent.bUsePawnViewRotation",NewFieldName="CameraComponent.bUsePawnControlRotation")
++K2FieldRedirects=(OldFieldName="SpringArmComponent.bUseControllerViewRotation",NewFieldName="SpringArmComponent.bUsePawnControlRotation")
++K2FieldRedirects=(OldFieldName="SpringArmComponent.bUsePawnViewRotation",NewFieldName="SpringArmComponent.bUsePawnControlRotation")
 
 +K2FieldRedirects=(OldFieldName="MovementComponent.GetModifiedMaxSpeed",NewFieldName="MovementComponent.K2_GetModifiedMaxSpeed")
 +K2FieldRedirects=(OldFieldName="MovementComponent.GetMaxSpeedModifier",NewFieldName="MovementComponent.K2_GetMaxSpeedModifier")
