- Fixed DeviceProfile comment

diff --git a/Engine/Config/BaseDeviceProfiles.ini b/Engine/Config/BaseDeviceProfiles.ini
--- a/Engine/Config/BaseDeviceProfiles.ini
+++ b/Engine/Config/BaseDeviceProfiles.ini
@@ -83,7 +83,7 @@ BaseProfileName=
 +CVars=r.TranslucentLightingVolume=0
 // Needs geometry shader support
 +CVars=r.AllowPointLightCubemapShadows=0
-// TXAA (?) looks blurry right now
+// Temporal AA (?) looks blurry right now
 +CVars=r.PostProcessAAQuality=0
 // Not sure on this one - may need the whole D Buffer stuff going
 +CVars=r.Decal.StencilSizeThreshold=-1
