Fixing localization paths.

diff --git a/Engine/Config/BaseEditor.ini b/Engine/Config/BaseEditor.ini
--- a/Engine/Config/BaseEditor.ini
+++ b/Engine/Config/BaseEditor.ini
@@ -1,8 +1,8 @@
 [Internationalization]
 +LocalizationPaths=../../../Engine/Content/Localization/Editor
 +LocalizationPaths=../../../Engine/Content/Localization/EditorTutorials
-+ToolTipLocalizationPaths=../../../Engine/Content/Localization/Editor/ToolTips
-+PropertyNameLocalizationPaths=../../../Engine/Content/Localization/Editor/PropertyNames
++ToolTipLocalizationPaths=../../../Engine/Content/Localization/ToolTips
++PropertyNameLocalizationPaths=../../../Engine/Content/Localization/PropertyNames
 ShouldLoadLocalizedPropertyNames=True
 
 [UnrealEd.GeneralSettings]
