Prevent vsync in editor from occuring when vsync is forced on via system settings

This is needed for the 4.8 effort asap

diff --git a/Engine/Config/BaseEngine.ini b/Engine/Config/BaseEngine.ini
--- a/Engine/Config/BaseEngine.ini
+++ b/Engine/Config/BaseEngine.ini
@@ -1656,6 +1656,9 @@ FLightPrimitiveInteractionInitialBlockSize=512
 [SystemSettingsEditor]
 ; System settings overrides for the editor.  Ideally the editor should use the same settings as the game.
 
+; Do not vsync in the editor. Throws off gpu profiling.
+r.VSync=0
+
 ; Parallel rendering has not been tested in the editor so just disable it.
 r.RHICmdBypass=1
 
