nicer autocomplete

diff --git a/Engine/Config/BaseInput.ini b/Engine/Config/BaseInput.ini
--- a/Engine/Config/BaseInput.ini
+++ b/Engine/Config/BaseInput.ini
@@ -49,6 +49,7 @@ MaxScrollbackSize=1024
 +AutoCompleteMapPaths=Content/Maps
 +ManualAutoCompleteList=(Command="Exit",Desc="Exit (Exits the game)")
 +ManualAutoCompleteList=(Command="DebugCreatePlayer 1",Desc=)
++ManualAutoCompleteList=(Command="ToggleDrawEvents",Desc="ToggleDrawEvents Toggles annotations for shader debugging with Pix or similar GPU capture tools")
 +ManualAutoCompleteList=(Command="Shot",Desc="Make a screenshot")
 +ManualAutoCompleteList=(Command="RecompileShaders changed",Desc="Recompile the shaders that have changed")
 +ManualAutoCompleteList=(Command="RecompileShaders global",Desc="Recompile all global shaders")
