Remove defunct 'ModifySoundClass' from console auto-complete list

diff --git a/Engine/Config/BaseInput.ini b/Engine/Config/BaseInput.ini
--- a/Engine/Config/BaseInput.ini
+++ b/Engine/Config/BaseInput.ini
@@ -123,7 +123,6 @@ MaxScrollbackSize=1024
 +ManualAutoCompleteList=(Command="IsolateDryAudio",Desc="IsolateDryAudio")
 +ManualAutoCompleteList=(Command="IsolateReverb",Desc="IsolateReverb")
 +ManualAutoCompleteList=(Command="ResetSoundState",Desc="ResetSoundState (Resets volumes to default and removes test filters)")
-+ManualAutoCompleteList=(Command="ModifySoundClass",Desc="ModifySoundClass <SoundClassName> Vol=<new volume>")
 +ManualAutoCompleteList=(Command="DisableAllScreenMessages",Desc="DisableAllScreenMessages (Disables all on-screen warnings/messages)")
 +ManualAutoCompleteList=(Command="EnableAllScreenMessages",Desc="EnableAllScreenMessages (Enables all on-screen warnings/messages)")
 +ManualAutoCompleteList=(Command="ToggleAllScreenMessages",Desc="ToggleAllScreenMessages (Toggles display state of all on-screen warnings/messages)")
