Add a commented option that works around crash on some Linux setups.

diff --git a/Engine/Config/ConsoleVariables.ini b/Engine/Config/ConsoleVariables.ini
--- a/Engine/Config/ConsoleVariables.ini
+++ b/Engine/Config/ConsoleVariables.ini
@@ -34,3 +34,6 @@
 ; Uncomment when running with a graphical debugger (but not when profiling)
 ;r.Shaders.Optimize=0
 ;r.Shaders.KeepDebugInfo=1
+
+; If Linux editor crashes in FMallocBinned with callstack that mentions MeshUtilities, you may need to uncomment this.
+;r.TriangleOrderOptimization=2
\ No newline at end of file
