Added config migration path for newer versions of the engine.

Newly installed versions of the engine will now attempt to copy the project-agnostic config settings from a previous engine installation. This happens by way of a versioned manifest that copies old versions when the manifest does not exist, or is a different version. This code path is benign for non-installed versions of the engine (or FPaths::ShouldSaveToUserDir() is false).

EditorGameAgnosticSettings and EditorUserSettings ini paths have been renamed to EditorSettings and EditorPerProjectUserSettings respectively to better convey their purpose. In general, most settings should be saved in EditorSettings (project-agnostic) so that they apply regardless of which project is open. We have some way to go migrating existing settings for this to be the case, however.

Some previously per-project configuration files are now project-agnostic (such as Editor.ini, EditorKeyBindings.ini, and EditorLayout.ini)

GEditor->Access...Settings and GEditor->Get...Settings have been removed in favor of direct access of the CDO through GetMutableDefault<> and GetDefault<> respectively. Global config ini filenames that are not set up are now neither loaded nor saved on build machines, to handle the problem of indeterminate state more generically.

This addresses UETOOL-270 (Most editor preferences should be project-agnostic)

diff --git a/Engine/Config/BaseEditorPerProjectUserSettings.ini b/Engine/Config/BaseEditorPerProjectUserSettings.ini
new file mode 100644
--- /dev/null
+++ b/Engine/Config/BaseEditorPerProjectUserSettings.ini
@@ -0,0 +1,693 @@
+
+; -------------------------------------------------------------------------------------------------
+; IMPORTANT:  EditorUserSettings.ini has special behavior!
+;
+; This .ini file contains the defaults for many editor and engine "preferences".  These preferences
+; are saved to an EditorUserSettings.ini file in the user's /<Game>/Saved/Config/ directory. 
+;
+; If you change the default for a preference in here, that change will *NOT* automatically 
+; propagate to users who already have saved preferences.  The idea is that we must preserve
+; their existing settings and never want to clobber entries in that file.
+;
+; If you add a new preference, the new key and value *WILL* be propagated to the user's
+; EditorUserSettings.ini file.  After we load the user's existing settings, we'll merge in
+; any missing keys and values that are present in these defaults.
+;
+; One easy technique for forcing a "reset" of an already-existing preference is to simply
+; rename the variable for that preference.  The config system will treat it as a newly-added
+; key and will propagate the change to existing users' preferences.
+; -------------------------------------------------------------------------------------------------
+
+[/Script/UnrealEd.EditorUserSettings]
+; True if WASD keys should be remapped to flight controls while the right mouse button is held down
+AllowFlightCameraToRemapKeys=True
+
+; The Blutility shelf holds editor utility Blueprints. Summon from the Workspace menu.
+bEnableEditorUtilityBlueprints=false
+
+;If this is true, the user will not be asked to fully load a package before saving or before creating a new object
+bSuppressFullyLoadPrompt=True
+
+; True if user should be allowed to select translucent objects in perspective viewports
+; NOTE: The reason we default this to true, even though it can be annoying, is that it can be a frustrating experience for new users to not be able to select a translucent object of interest, until they at least learn about this feature.
+bAllowSelectTranslucent=True
+
+; True if we should move actors to their appropriate grid volume levels immediately after most operations
+bUpdateActorsInGridLevelsImmediately=False
+
+; True if the memory/size data is displayed in the slate level browser
+bDisplayMemorySizeDataInLevelBrowser=False
+	
+; When enabled, the application frame rate, memory, Unreal object count, as well as a debug console will be displayed in the main editor UI
+bShowFrameRateAndMemory=False
+
+; Lowers CPU usage when the editor is in the background and not the active application
+bThrottleCPUWhenNotForeground=True
+
+; True if we should automatically reimport animset assets when a change to source content is detected
+bAutoReimportAnimSets=False
+
+; Property Matrix
+PropertyMatrix_NumberOfPasteOperationsBeforeWarning=20
+
+; Blueprint Editor SCS Editor settings
+bSCSEditorShowGrid=true
+
+; Export
+bKeepAttachHierarchy=true
+
+; Misc
+HoverHighlightIntensity=0.000000
+bDisplayUIExtensionPoints=False
+bUseCurvesForDistributions=False
+bAutoloadCheckedOutPackages=False
+bAutomaticallyHotReloadNewClasses=True
+
+[/Script/EditorStyle.EditorStyleSettings]
+; Applies a color vision deficiency filter to the entire editor
+ColorVisionDeficiencyType=CVD_NormalVision
+
+SelectionColor=(R=0.728f,G=0.364f,B=0.003f,A=1.000000)
+PressedSelectionColor=(R=0.701f,G=0.225f,B=0.003f,A=1.000000)
+InactiveSelectionColor=(R=0.25f,G=0.25f,B=0.25f,A=1.000000)
+SelectorColor=(R=0.701f,G=0.225f,B=0.003f,A=1.000000)
+
+; If true, all toolbars will use small icons without labels
+bUseSmallToolBarIcons=False
+; Enables animated transitions for certain menus and pop-up windows.  Note that animations may be automatically disabled at low frame rates in order to improve responsiveness.
+bEnableWindowAnimations=False
+;If true, Variables names in the script editor are displayed in a sanitized format
+bShowFriendlyNames=True
+;If true, the Editor Preferences and Project Settings menu entries in the main menu will be expanded with sub-menus for each settings section
+bExpandConfigurationMenus=False
+bShowProjectMenus=True
+bShowLaunchMenus=True
+
+[/Script/Levels.LevelBrowserSettings]
+; True if the actor count is displayed in the slate level browser
+bDisplayActorCount=True
+bDisplayLightmassSize=False
+bDisplayFileSize=False
+; True if Level Paths are displayed in the slate level browser
+bDisplayPaths=False
+bDisplayEditorOffset=False
+
+[/Script/SceneOutliner.SceneOutlinerSettings]
+; True when the Scene Outliner is hiding temporary/run-time Actors
+bHideTemporaryActors=False
+; True when the Scene Outliner is showing only Actors that exist in the current level
+bShowOnlyActorsInCurrentLevel=False
+; Scene Outliner Settings
+bShowOnlySelectedActors=False
+
+
+[/Script/UnrealEd.ContentBrowserSettings]
+; The number of objects to load at once in the Content Browser before displaying a warning about loading many assets
+NumObjectsToLoadBeforeWarning=20
+; Whether to render thumbnails for loaded assets in real-time in the Content Browser
+RealTimeThumbnails=True
+; Whether to display folders in the assets view of the content browser. Note that this implies 'Show Only Assets in Selected Folders'.
+DisplayFolders=True
+; Whether to display the engine folders in the assets view of the content browser. Note that this implies 'Display Folders'.
+DisplayEngineFolder=False
+; Whether to display the developers folder in the path or assets view of the content browser
+DisplayDevelopersFolder=False
+; Whether to display the collections view in the Content Browser
+DisplayCollections=False
+
+[/Script/UnrealEd.DestructableMeshEditorSettings]
+; Persona viewport default settings
+AnimPreviewFloorColor=(B=6,G=24,R=43,A=255)
+AnimPreviewSkyColor=(B=250,G=196,R=178,A=255)
+AnimPreviewSkyBrightness=0.250000
+AnimPreviewLightBrightness=1.000000
+AnimPreviewLightingDirection=(Pitch=320.000000,Yaw=290.000000,Roll=0.000000)
+AnimPreviewDirectionalColor=(B=253,G=253,R=253,A=255)
+
+
+[/Script/UnrealEd.EditorExperimentalSettings]
+; Epic Labs
+bWorldBrowser=False
+bBreakOnExceptions=False
+bToolbarCustomization=False
+bBehaviorTreeEditor=False
+
+
+[/Script/UnrealEd.EditorLoadingSavingSettings]
+; True if we should automatically load a default level at start up
+LoadLevelAtStartup=ProjectDefault
+; True if we should automatically reimport textures when a change to source content is detected
+bAutoReimportTextures=False
+bAutoReimportCSV=False
+; Whether to mark blueprints dirty if they are automatically migrated during loads
+bDirtyMigratedBlueprints=False
+; Whether to automatically save after a time interval
+bAutoSaveEnable=True
+; Whether to automatically save maps during an autosave
+bAutoSaveMaps=True
+; Whether to automatically save content packages during an autosave
+bAutoSaveContent=True
+; The time interval after which to save
+AutoSaveTimeMinutes=10
+; The number of seconds warning before an autosave
+AutoSaveWarningInSeconds=10
+
+; Whether to automatically perform an SCC checkout on asset modification
+bPromptForCheckoutOnAssetModification=False
+; Whether to automatically prompt for SCC checkout on asset modification
+bPromptForCheckoutOnAssetModification=True
+; Source control
+bSCCAutoAddNewFiles=True
+; Tool to use for diff'ing text
+TextDiffToolPath=(FilePath="p4merge.exe")
+
+
+[/Script/UnrealEd.LevelEditorMiscSettings]
+; Auto apply lighting after it has been built
+bAutoApplyLightingEnable=true
+; If true, BSP will auto-update
+bBSPAutoUpdate=True
+; If true, Navigation will auto-update
+bNavigationAutoUpdate=True
+; If true, Replaces respects the scale of the original actor. If false, Replace sets the scale to 1.0
+bReplaceRespectsScale=True
+
+; If true audio will continue to play when the editor does not have focus
+bAllowBackgroundAudio=False
+; If true audio will be enabled in the editor. Does not affect PIE 
+bEnableRealTimeAudio=False
+; Volume level for the editor
+EditorVolumeLevel=1.0
+; Enables audio feedback for certain operations in Unreal Editor, such as entering and exiting Play mode
+bEnableEditorSounds=True
+; The default level streaming class to use when adding new streaming levels
+DefaultLevelStreamingClass=Class'/Script/Engine.LevelStreamingKismet'
+
+[/Script/UnrealEd.LevelEditorPlaySettings]
+PlayFromHerePlayerStartClassName=/Script/Engine.PlayerStartPIE
+
+EnableSound=true
+NewWindowWidth=1280
+NewWindowHeight=720
+NewWindowPosition=(X=0,Y=0)
+StandaloneWindowWidth=1280
+StandaloneWindowHeight=720
+ShowMouseControlLabel=True
+AutoRecompileBlueprints=True
+; True if Play In Editor should only load currently-visible levels in PIE
+bOnlyLoadVisibleLevelsInPIE=False
+LastExecutedPlayModeLocation=PlayLocation_DefaultPlayerStart
+LastExecutedPlayModeType=PlayMode_InViewPort
+LastExecutedLaunchModeType=LaunchMode_OnDevice
+LastExecutedLaunchPlatform=
+
+; common screen resolutions for laptops
++LaptopScreenResolutions=(Description="Apple MacBook Air 11",Width=1366,Height=768,AspectRatio="16:9")
++LaptopScreenResolutions=(Description="Apple MacBook Air 13\"",Width=1440,Height=900,AspectRatio="16:10")
++LaptopScreenResolutions=(Description="Apple MacBook Pro 13\"",Width=1280,Height=800,AspectRatio="16:10")
++LaptopScreenResolutions=(Description="Apple MacBook Pro 13\" (Retina)",Width=2560,Height=1600,AspectRatio="16:10")
++LaptopScreenResolutions=(Description="Apple MacBook Pro 15\"",Width=1440,Height=900,AspectRatio="16:10")
++LaptopScreenResolutions=(Description="Apple MacBook Pro 15\" (Retina)",Width=2880,Height=1800,AspectRatio="16:10")
++LaptopScreenResolutions=(Description="Generic 14-15.6\" Notebook",Width=1366,Height=768,AspectRatio="16:9")
+
+; common screen resolutions for desktop monitors
++MonitorScreenResolutions=(Description="19\" monitor",Width=1440,Height=900,AspectRatio="16:10")
++MonitorScreenResolutions=(Description="20\" monitor",Width=1600,Height=900,AspectRatio="16:9")
++MonitorScreenResolutions=(Description="22\" monitor",Width=1680,Height=1050,AspectRatio="16:10")
++MonitorScreenResolutions=(Description="21.5-24\" monitor",Width=1920,Height=1080,AspectRatio="16:9")
++MonitorScreenResolutions=(Description="27\" monitor",Width=2560,Height=1440,AspectRatio="16:9")
+
+; common screen resolutions for mobile phones
++PhoneScreenResolutions=(Description="Apple iPhone 4, 4S, iTouch 4 (Portrait)",Width=640,Height=960,AspectRatio="2:3")
++PhoneScreenResolutions=(Description="Apple iPhone 4, 4S, iTouch 4 (Landscape)",Width=960,Height=640,AspectRatio="3:2")
++PhoneScreenResolutions=(Description="Apple iPhone 5, 5S, iTouch 5 (Portrait)",Width=640,Height=1136,AspectRatio="~9:16")
++PhoneScreenResolutions=(Description="Apple iPhone 5, 5S, iTouch 5 (Landscape)",Width=1136,Height=640,AspectRatio="~16:9")
++PhoneScreenResolutions=(Description="HTC One (Portrait)",Width=1080,Height=1920,AspectRatio="9:16")
++PhoneScreenResolutions=(Description="HTC One (Landscape)",Width=1920,Height=1080,AspectRatio="16:9")
++PhoneScreenResolutions=(Description="Samsung Galaxy S4 (Portrait)",Width=1080,Height=1920,AspectRatio="9:16")
++PhoneScreenResolutions=(Description="Samsung Galaxy S4 (Landscape)",Width=1920,Height=1080,AspectRatio="16:9")
+
+; common screen resolutions for tablet devices
++TabletScreenResolutions=(Description="Apple iPad 2, iPad Mini (Portrait)",Width=768,Height=1024,AspectRatio="~3:4")
++TabletScreenResolutions=(Description="Apple iPad 2, iPad Mini (Landscape)",Width=1024,Height=768,AspectRatio="~4:3")
++TabletScreenResolutions=(Description="Apple iPad 3, 4, Air (Portrait)",Width=1536,Height=2048,AspectRatio="3:4")
++TabletScreenResolutions=(Description="Apple iPad 3, 4, Air (Landscape)",Width=2048,Height=1536,AspectRatio="4:3")
++TabletScreenResolutions=(Description="Microsoft Surface RT (Portrait)",Width=768,Height=1366,AspectRatio="9:16")
++TabletScreenResolutions=(Description="Microsoft Surface RT (Landscape)",Width=1366,Height=768,AspectRatio="16:9")
++TabletScreenResolutions=(Description="Microsoft Surface Pro (Portrait)",Width=1080,Height=1920,AspectRatio="9:16")
++TabletScreenResolutions=(Description="Microsoft Surface Pro (Landscape)",Width=1920,Height=1080,AspectRatio="16:9")
+
+; common screen resolutions for television sets
++TelevisionScreenResolutions=(Description="720p (HDTV, Blu-ray)",Width=1280,Height=720,AspectRatio="16:9")
++TelevisionScreenResolutions=(Description="1080i, 1080p (HDTV, Blu-ray)",Width=1920,Height=1080,AspectRatio="16:9")
++TelevisionScreenResolutions=(Description="4K Ultra HD",Width=3840,Height=2160,AspectRatio="16:9")
++TelevisionScreenResolutions=(Description="4K Digital Cinema",Width=4096,Height=2160,AspectRatio="1.90:1")
+
+
+[/Script/UnrealEd.LevelEditorViewportSettings]
+; Use WASD flight camera controls by default
+FlightCameraControlType=WASD_RMBOnly
+; Ignore Ctrl key by default in the landscape/foliage editors
+LandscapeEditorControlType=IgnoreCtrl
+FoliageEditorControlType=IgnoreCtrl
+; If true, moves the canvas and shows the mouse.  If false, uses original camera movement
+bPanMovesCanvas=True
+; If true, zooms centering on the mouse position.  If false, the zoom is around the center of the viewport
+bCenterZoomAroundCursor=True
+bAllowTranslateRotateZWidget=False
+; If true, Clicking a BSP selects the brush and ctrl+shift+click selects the surface. If false, vice versa
+bClickBSPSelectsBrush=True
+;Mouse speed when dragging in viewport
+CameraSpeed=4
+;Mouse speed when using middle mouse scroll in viewport
+MouseScrollCameraSpeed=5
+MouseSensitivty=.2f
+; Whether to use mouse position as direct widget position
+bUseAbsoluteTranslation=True
+bLevelStreamingVolumePrevis=False
+bUseUE3OrbitControls=False
+
+bUsePowerOf2SnapSize=False
+
+.DecimalGridSizes=1
+.DecimalGridSizes=5
+.DecimalGridSizes=10
+.DecimalGridSizes=50
+.DecimalGridSizes=100
+.DecimalGridSizes=500
+.DecimalGridSizes=1000
+.DecimalGridSizes=5000
+.DecimalGridSizes=10000
+
+.DecimalGridIntervals=10.000000
+.DecimalGridIntervals=5.000000
+.DecimalGridIntervals=10.000000
+.DecimalGridIntervals=5.000000
+.DecimalGridIntervals=10.000000
+.DecimalGridIntervals=5.000000
+.DecimalGridIntervals=10.000000
+.DecimalGridIntervals=5.000000
+.DecimalGridIntervals=10.000000
+
+.Pow2GridSizes=1
+.Pow2GridSizes=2
+.Pow2GridSizes=4
+.Pow2GridSizes=8
+.Pow2GridSizes=16
+.Pow2GridSizes=32
+.Pow2GridSizes=64
+.Pow2GridSizes=128
+.Pow2GridSizes=256
+.Pow2GridSizes=512
+.Pow2GridSizes=1024
+.Pow2GridSizes=2048
+.Pow2GridSizes=4096
+.Pow2GridSizes=8192
+.Pow2GridIntervals=8
+
+.CommonRotGridSizes=5
+.CommonRotGridSizes=10
+.CommonRotGridSizes=15
+.CommonRotGridSizes=30
+.CommonRotGridSizes=60
+.CommonRotGridSizes=120
+.DivisionsOf360RotGridSizes=2.8125
+.DivisionsOf360RotGridSizes=5.625
+.DivisionsOf360RotGridSizes=11.25
+.DivisionsOf360RotGridSizes=22.5
+.DivisionsOf360RotGridSizes=45
+.DivisionsOf360RotGridSizes=90
+.ScalingGridSizes=10
+.ScalingGridSizes=1
+.ScalingGridSizes=0.5
+.ScalingGridSizes=0.25
+.ScalingGridSizes=0.125
+.ScalingGridSizes=0.0625
+.ScalingGridSizes=0.03125
+
+GridEnabled=True
+RotGridEnabled=True
+SnapScaleEnabled=True
+bSnapNewObjectsToFloor=True
+; If enabled, use the old-style multiplicative/percentage scaling method instead of the new additive/fraction method
+bUsePercentageBasedScaling=False
+; If true actor snap will be enabled in the editor
+bEnableActorSnap=False
+; Actor snap scale for the editor
+ActorSnapScale=1.0
+; Actor snap distance setting for the editor
+ActorSnapDistance=100.0
+bSnapVertices=False
+SnapDistance=10.000000
+CurrentPosGridSize=2
+CurrentRotGridSize=1
+CurrentScalingGridSize=3
+CurrentRotGridMode=GridMode_Common
+
+;default to maintaining fov along y-axis
+AspectRatioAxisConstraint=AspectRatio_MaintainXFOV
+; Enables real-time hover feedback when mousing over objects in editor viewports
+bEnableViewportHoverFeedback=False
+; If enabled, selected objects will be highlighted with brackets in all modes rather than a special highlight color.
+bHighlightWithBrackets=False
+; If true all orthographic viewports are linked to the same position and move together
+bUseLinkedOrthographicViewports=True
+; If true, objects must be entirely encompassed by the selection box in ortho. viewports to be selected
+bStrictBoxSelection=False
+; True if viewport box selection also selects occluded objects, false if only objects with visible pixels are selected
+bTransparentBoxSelection=False
+; If enabled, selected objects will have an outline around them
+bUseSelectionOutline=True
+; Sets the intensity of the overlay displayed when an object is selected (defaults to be 0 so you can see the material of selected objects)
+SelectionHighlightIntensity=0.0
+; Sets the intensity of the overlay displayed when a BSP surface is selected (defaults to be 0 so you can see the material of selected objects)
+BSPSelectionHighlightIntensity=0.2
+; Enables the editor perspective camera to be dropped at the last PlayInViewport cam position
+bEnableViewportCameraToUpdateFromPIV=True
+; When enabled, selecting a camera actor will display a live 'picture in picture' preview from the camera's perspective within the current editor viewport.  This can be used to easily tweak camera positioning, post-processing and other settings without having to possess the camera itself.  This feature may reduce application performance when enabled. */
+bPreviewSelectedCameras=True
+; Affects the size of 'picture in picture' previews if they are enabled
+CameraPreviewSize=5.0
+; This distance is used to place actors which are dropped on nothing in the viewport
+BackgroundDropDistance=768
+; When enabled, simple stats that are enabled in level viewports are preserved between editor sessions
+bSaveSimpleStats=False
+
+
+[ColorPickerUI]
+bAdvancedSectionExpanded=False
+bSRGBEnabled=True
+bWheelMode=True
+
+
+[Undo]
+; Size of Undo buffer in MB. Bigger number allows more Undo history especially when working with Landscape
+UndoBufferSize=32
+
+
+[PropertySettings]
+ShowFriendlyPropertyNames=True
+ExpandDistributions=false
+
+[MRU]
+
+
+[/Script/UnrealEd.MaterialEditorOptions]
+bShowGrid=True
+bShowBackground=False
+bHideUnusedConnectors=False
+bRealtimeMaterialViewport=True
+bRealtimeExpressionViewport=False
+bAlwaysRefreshAllPreviews=False
+bLivePreviewUpdate=True
+
+
+[UnEdViewport]
+InterpEdPanInvert=False
+
+
+[FEditorModeTools]
+ShowWidget=True
+CoordSystem=0
+UseAbsoluteTranslation=True
+AllowTranslateRotateZWidget=False
+
+[LightingBuildOptions]
+OnlyBuildSelectedActors=false
+OnlyBuildCurrentLevel=false
+OnlyBuildChanged=false
+BuildBSP=true
+BuildActors=true
+QualityLevel=0
+NumUnusedLocalCores=1
+ShowLightingBuildInfo=false
+
+
+[MatineeCreateMovieOptions]
+CloseEditor=false
+CaptureResolutionIndex = 0;
+CaptureResolutionFPS = 30;
+CaptureTypeIndex = 0;
+Compress=false
+CinematicMode=true
+DisableMovement=true
+DisableTurning=true
+HidePlayer=true
+DisableInput=true
+HideHUD=true
+
+
+[Matinee]
+Hide3DTracks=false
+ZoomToScrubPos=false
+ShowCurveEd=false
+
+
+[/Script/UnrealEd.PhATSimOptions]
+SkyBrightness=0.25
+; for physical correct light computations we multiply diffuse and specular lights by PI (see LABEL_RealEnergy)
+Brightness=3.1415926535897932
+AngularSnap=15.0
+LinearSnap=2.0
+SimSpeed=1.0
+bDrawContacts=false
+FloorGap=25.0
+GravScale=1.0
+bPromptOnBoneDelete=true
+PokeStrength=100.0
+bShowNamesInHierarchy=true
+PokePauseTime=0.5
+PokeBlendTime=0.5
+ConstraintDrawSize=0.01
+
+
+[/Script/UnrealEd.CurveEdOptions]
+MinViewRange=0.01
+MaxViewRange=1000000.0
+BackgroundColor=(R=0.23529412,G=0.23529412,B=0.23529412,A=1.0)
+LabelColor=(R=0.4,G=0.4,B=0.4,A=1.0)
+SelectedLabelColor=(R=0.6,G=0.4,B=0.1, A=1.0)
+GridColor=(R=0.35,G=0.35,B=0.35,A=1.0)
+GridTextColor=(R=0.78431373,G=0.78431373,B=0.78431373,A=1.0)
+LabelBlockBkgColor=(R=0.25,G=0.25,B=0.25,A=1.0)
+SelectedKeyColor=(R=1.0,G=1.0,B=0.0,A=1.0)
+
+
+[/Script/UnrealEd.PersonaOptions]
+bShowGrid=False
+bHighlightOrigin=True
+bShowSky=True
+bShowFloor=True
+GridSize=25
+ViewModeType=2
+ViewportBackgroundColor=(R=0.04,G=0.04,B=0.04)
+ViewFOV=53.43
+ShowMeshStats=1
+
+
+[UnrealEd.UIEditorOptions]
+WindowPosition=(X=256,Y=256,Width=1024,Height=768)
+ViewportSashPosition=824
+PropertyWindowSashPosition=568
+ViewportGutterSize=0
+VirtualSizeX=0
+VirtualSizeY=0
+bRenderViewportOutline=true
+bRenderContainerOutline=true
+bRenderSelectionOutline=true
+bRenderSelectionHandles=true
+bRenderPerWidgetSelectionOutline=true
+GridSize=8
+bSnapToGrid=true
+mViewDrawGrid=true
+bShowDockHandles=true
+
+
+[/Script/UnrealEd.CascadeOptions]
+bShowModuleDump=false
+BackgroundColor=(B=25,G=20,R=20,A=0)
+bUseSubMenus=true
+bUseSpaceBarReset=false
+bUseSpaceBarResetInLevel=true
+Empty_Background=(B=25,G=20,R=20,A=0)
+Emitter_Background=(B=25,G=20,R=20,A=0)
+Emitter_Unselected=(B=0,G=100,R=255,A=0)
+Emitter_Selected=(B=180,G=180,R=180,A=0)
+ModuleColor_General_Unselected=(B=49,G=40,R=40,A=0)
+ModuleColor_General_Selected=(B=0,G=100,R=255,A=0)
+ModuleColor_TypeData_Unselected=(B=20,G=20,R=15,A=0)
+ModuleColor_TypeData_Selected=(B=0,G=100,R=255,A=0)
+ModuleColor_Beam_Unselected=(R=160,G=150,B=235)
+ModuleColor_Beam_Selected=(R=255,G=100,B=0)
+ModuleColor_Trail_Unselected=(R=130,G=235,B=170)
+ModuleColor_Trail_Selected=(R=255,G=100,B=0)
+ModuleColor_Spawn_Unselected=(R=200,G=100,B=100)
+ModuleColor_Spawn_Selected=(R=255,G=50,B=50)
+ModuleColor_Light_Unselected=(B=49,G=40,R=90)
+ModuleColor_Light_Selected=(B=0,G=100,R=255)
+ModuleColor_Required_Unselected=(R=200,G=200,B=100)
+ModuleColor_Required_Selected=(R=255,G=225,B=50)
+ModuleColor_Event_Unselected=(R=64,G=64,B=255)
+ModuleColor_Event_Selected=(R=0,G=0,B=255)
+bShowGrid=false
+GridColor_Hi=(R=0,G=100,B=255)
+GridColor_Low=(R=0,G=100,B=255)
+GridPerspectiveSize=1024
+ShowPPFlags=0
+bUseSlimCascadeDraw=true
+SlimCascadeDrawHeight=24
+bCenterCascadeModuleText=true
+Cascade_MouseMoveThreshold=4
+MotionModeRadius=150.0
+
+
+[ContentBrowserFilter]
+FavoriteTypes_1=Animation Sequence;Material Instances (Constant);Materials;Particle Systems;Skeletal Meshes;Sound Cues;Static Meshes;Textures;Blueprint
+
+
+[FAutoPackageBackup]
+Enabled=True
+MaxAllowedSpaceInMB=250
+BackupIntervalInMinutes=5
+
+
+[/Script/UnrealEd.FbxImportUI]
+bImportMaterials=True
+bAutoCreateGroups=True
+bInvertNormalMaps=False
+bImportTextures=True
+bOverrideFullName=True
+bConvertScene=True
+bRemoveNamespaces=True
+bImportAnimations=False
+bResampleAnimations=False
+bImportRigidMesh=False
+bCombineMeshes=True
+bCreatePhysicsAsset=True
+bImportMesh=True
+
+
+[/Script/UnrealEd.FbxStaticMeshImportData]
+bOneConvexHullPerUCX=True
+bImportMeshLODs=False
+NormalImportMethod=FBXNIM_ImportNormals
+VertexColorImportOption=EVertexColorImportOption::Ignore
+VertexOverrideColor=(R=255,G=255,B=255,A=255)
+
+
+[/Script/UnrealEd.FbxSkeletalMeshImportData]
+bPreserveSmoothingGroups=True
+bImportMeshLODs=False
+bImportMorphTargets=False
+
+
+[/Script/UnrealEd.FbxTextureImportData]
+
+
+[SoundSettings]
+ChirpSoundClasses=Dialog DialogMedium DialogLoud DialogDeafening
+BatchProcessMatureNodeSoundClasses=Dialog Chatter
+
+
+[EditorPreviewMesh]
+; Preview static meshes used in the editor. 
+PreviewMeshNames="/Engine/EditorMeshes/ColorCalibrator/SM_ColorCalibrator.SM_ColorCalibrator"
+
+
+[EditorLayout]
+SlateMainFrameLayout=""
+
+
+[LandscapeEdit]
+ToolStrength=0.300000
+WeightTargetValue=1.000000
+bUseWeightTargetValue=False
+BrushRadius=2048.000000
+BrushComponentSize=1
+BrushFalloff=0.500000
+bUseClayBrush=False
+AlphaBrushScale=0.500000
+AlphaBrushRotation=0.000000
+AlphaBrushPanU=0.500000
+AlphaBrushPanV=0.500000
+AlphaTextureName=/Engine/EditorLandscapeResources/DefaultAlphaTexture.DefaultAlphaTexture
+AlphaTextureChannel=0
+FlattenMode=0
+bUseSlopeFlatten=False
+bPickValuePerApply=False
+ErodeThresh=64
+ErodeIterationNum=28
+ErodeSurfaceThickness=256
+ErosionNoiseMode=2
+ErosionNoiseScale=60.000000
+RainAmount=128
+SedimentCapacity=0.300000
+HErodeIterationNum=28
+RainDistMode=0
+RainDistScale=60.000000
+HErosionDetailScale=0.010000
+bHErosionDetailSmooth=True
+NoiseMode=0
+NoiseScale=128.000000
+SmoothFilterKernelScale=1.000000
+DetailScale=0.300000
+bDetailSmooth=False
+MaximumValueRadius=10000.000000
+bSmoothGizmoBrush=True
+PasteMode=0
+ConvertMode=0
+bApplyToAllTargets=True
+
+
+[FoliageEdit]
+Radius=512.000000
+PaintDensity=0.500000
+UnpaintDensity=0.000000
+bFilterLandscape=True
+bFilterStaticMesh=True
+bFilterBSP=True
+bFilterTranslucent=False
+
+
+[MeshPaintEdit]
+DefaultBrushRadius=128
+
+
+[BlueprintSpawnNodes]
+; Comment box is bound to C, but that is handled differently due to it needing to work without clicking
++Node=(Class=Actor:ReceiveBeginPlay Key=P Shift=false Ctrl=false Alt=false)
++Node=(Class="Do N" Key=N Shift=false Ctrl=false Alt=false)
++Node=(Class=KismetSystemLibrary:Delay Key=D Shift=false Ctrl=false Alt=false)
++Node=(Class=K2Node_IfThenElse Key=B Shift=false Ctrl=false Alt=false)
++Node=(Class=K2Node_ExecutionSequence Key=S Shift=false Ctrl=false Alt=false)
++Node=(Class=Gate Key=G Shift=false Ctrl=false Alt=false)
++Node=(Class=K2Node_MultiGate Key=M Shift=false Ctrl=false Alt=false)
++Node=(Class=ForEachLoop Key=F Shift=false Ctrl=false Alt=false)
++Node=(Class=DoOnce Key=O Shift=false Ctrl=false Alt=false)
++Node=(Class=KismetArrayLibrary:Array_Get Key=A Shift=false Ctrl=false Alt=false)
+
+
+[MaterialEditorSpawnNodes]
++Node=(Class=MaterialExpressionAdd Key=A Shift=false Ctrl=false Alt=false)
++Node=(Class=MaterialExpressionBumpOffset Key=B Shift=false Ctrl=false Alt=false)
++Node=(Class=MaterialExpressionDivide Key=D Shift=false Ctrl=false Alt=false)
++Node=(Class=MaterialExpressionPower Key=E Shift=false Ctrl=false Alt=false)
++Node=(Class=MaterialExpressionMaterialFunctionCall Key=F Shift=false Ctrl=false Alt=false)
++Node=(Class=MaterialExpressionIf Key=I Shift=false Ctrl=false Alt=false)
++Node=(Class=MaterialExpressionLinearInterpolate Key=L Shift=false Ctrl=false Alt=false)
++Node=(Class=MaterialExpressionMultiply Key=M Shift=false Ctrl=false Alt=false)
++Node=(Class=MaterialExpressionNormalize Key=N Shift=false Ctrl=false Alt=false)
++Node=(Class=MaterialExpressionOneMinus Key=O Shift=false Ctrl=false Alt=false)
++Node=(Class=MaterialExpressionPanner Key=P Shift=false Ctrl=false Alt=false)
++Node=(Class=MaterialExpressionReflectionVectorWS Key=R Shift=false Ctrl=false Alt=false)
++Node=(Class=MaterialExpressionScalarParameter Key=S Shift=false Ctrl=false Alt=false)
++Node=(Class=MaterialExpressionTextureSample Key=T Shift=false Ctrl=false Alt=false)
++Node=(Class=MaterialExpressionTextureCoordinate Key=U Shift=false Ctrl=false Alt=false)
++Node=(Class=MaterialExpressionVectorParameter Key=V Shift=false Ctrl=false Alt=false)
++Node=(Class=MaterialExpressionConstant Key=One Shift=false Ctrl=false Alt=false)
++Node=(Class=MaterialExpressionConstant2Vector Key=Two Shift=false Ctrl=false Alt=false)
++Node=(Class=MaterialExpressionConstant3Vector Key=Three Shift=false Ctrl=false Alt=false)
++Node=(Class=MaterialExpressionConstant4Vector Key=Four Shift=false Ctrl=false Alt=false)
+
+[WidgetTemplatesExpanded]
+Common=True
\ No newline at end of file
diff --git a/Engine/Config/BaseEditorSettings.ini b/Engine/Config/BaseEditorSettings.ini
new file mode 100644
--- /dev/null
+++ b/Engine/Config/BaseEditorSettings.ini
@@ -0,0 +1,46 @@
+[/Script/UnrealEd.EditorSettings]
+bCopyStarterContentPreference = true
+
+[/Script/NewsFeed.NewsFeedSettings]
+CdnSourceUrl="http://cdn.unrealengine.com/newsfeed/[[Culture]]/index.json"
+LocalSourcePath="Editor/NewsFeed/[[Culture]]/"
+Source=NEWSFEED_Cdn
+MaxItemsToShow=10
+ShowOnlyUnreadItems=True
+
+[/Script/EditorLiveStreaming.EditorLiveStreamingSettings]
+FrameRate=30
+ScreenScaling=0.75
+bPrimaryMonitorOnly=true
+bEnableWebCam=true
+WebCamResolution=Normal_640x480
+bMirrorWebCamImage=false
+bCaptureAudioFromComputer=true
+bCaptureAudioFromMicrophone=true
+
+[/Script/IntroTutorials.EditorTutorialSettings]
++Categories=(Identifier="Basics",Title="Basics",Description="Getting started with Unreal Engine 4.",Icon="PlayWorld.RepeatLastPlay",Texture=None,SortOrder=100)
++Categories=(Identifier="Blueprints",Title="Blueprints",Description="Tutorials covering the creation and usage of Blueprints.",Icon="FullBlueprintEditor.EditGlobalOptions",Texture=/Engine/Tutorial/BlueprintTutorials/TutorialAssets/Blueprint_64x.Blueprint_64x,SortOrder=200)
++Categories=(Identifier="Editors",Title="Editors",Description="Tutorials covering the various sub editors of Unreal Engine 4.",Icon=,Texture=/Engine/Tutorial/SubEditors/TutorialAssets/icon_Editor_Preferences_General_40x.icon_Editor_Preferences_General_40x,SortOrder=225)
++Categories=(Identifier="Code",Title="Code",Description="Write C++ code for use in Unreal Engine 4.",Icon="PlayWorld.RepeatLastPlay",Texture=None,SortOrder=250)
++Categories=(Identifier="Paper2D",Title="Paper2D",Description="Tutorials convering the usage of Unreal Engine 4\'s 2D game tools: Paper2D.",Icon="ClassThumbnail.PaperSprite",Texture=/Engine/Tutorial/Paper2D/TutorialAssets/Paper2DSprite_TutorialIcon.Paper2DSprite_TutorialIcon,SortOrder=500)
++Categories=(Identifier="Animation",Title="Animation",Description="Tutorials convering the animation system in Unreal Engine 4.",Icon=,Texture=/Engine/Tutorial/SubEditors/TutorialAssets/icon_ShowSkeletalMeshes_40x.icon_ShowSkeletalMeshes_40x,SortOrder=600)
++Categories=(Identifier="Landscape",Title="Landscape",Description="Tutorials covering the Unreal Editor 4 terrain editor: Landscape.",Icon="LevelEditor.LandscapeMode",Texture=/Engine/Tutorial/Landscape/TutorialAssets/Landscape.Landscape,SortOrder=700)
++Categories=(Identifier="Foliage",Title="Foliage",Description="Tutorials covering the Unreal Engine 4 Foliage tool.",Icon="LevelEditor.FoliageMode",Texture=/Engine/Tutorial/Foliage/TutorialAssets/Foliage.Foliage,SortOrder=800)
++Categories=(Identifier="Mobile",Title="Mobile",Description="Mobile Tutorials.",Icon="MaterialEditor.ToggleMobileStats",Texture=None,SortOrder=900)
++StartupTutorial=/Engine/Tutorial/Basics/LevelEditorAttract.LevelEditorAttract_C
++TutorialContexts=(Context="StaticMeshEditor",BrowserFilter=,AttractTutorial=None,LaunchTutorial=/Engine/Tutorial/SubEditors/StaticMeshEditorTutorial.StaticMeshEditorTutorial_C)
++TutorialContexts=(Context="LevelEditor",BrowserFilter=,AttractTutorial=None,LaunchTutorial=/Engine/Tutorial/Basics/LevelEditorOverview.LevelEditorOverview_C)
++TutorialContexts=(Context="BlueprintEditor",BrowserFilter=,AttractTutorial=None,LaunchTutorial=/Engine/Tutorial/BlueprintTutorials/BlueprintEditorTutorial.BlueprintEditorTutorial_C)
++TutorialContexts=(Context="BlueprintEditor.MacroLibrary",BrowserFilter=,AttractTutorial=None,LaunchTutorial=/Engine/Tutorial/BlueprintTutorials/BlueprintMacroLibrariesEditorOverview.BlueprintMacroLibrariesEditorOverview_C)
++TutorialContexts=(Context="BlueprintEditor.Interface",BrowserFilter=,AttractTutorial=None,LaunchTutorial=/Engine/Tutorial/BlueprintTutorials/BlueprintInterfacesEditorOverview.BlueprintInterfacesEditorOverview_C)
++TutorialContexts=(Context="BlueprintEditor.LevelScript",BrowserFilter=,AttractTutorial=None,LaunchTutorial=/Engine/Tutorial/BlueprintTutorials/LevelBlueprintEditorOverview.LevelBlueprintEditorOverview_C)
++TutorialContexts=(Context="DestructibleMeshEditor",BrowserFilter=,AttractTutorial=None,LaunchTutorial=/Engine/Tutorial/SubEditors/DestructibleMeshEditorTutorial.DestructibleMeshEditorTutorial_C)
++TutorialContexts=(Context="MaterialEditor",BrowserFilter=,AttractTutorial=None,LaunchTutorial=/Engine/Tutorial/SubEditors/MaterialEditorTutorial.MaterialEditorTutorial_C)
++TutorialContexts=(Context="Cascade",BrowserFilter=,AttractTutorial=None,LaunchTutorial=/Engine/Tutorial/SubEditors/ParticleSystemEditorTutorial.ParticleSystemEditorTutorial_C)
++TutorialContexts=(Context="Persona",BrowserFilter=,AttractTutorial=None,LaunchTutorial=/Engine/Tutorial/SubEditors/PersonaAnimEditorTutorial.PersonaAnimEditorTutorial_C)
++TutorialContexts=(Context="SpriteEditor",BrowserFilter=,AttractTutorial=None,LaunchTutorial=/Engine/Tutorial/Paper2D/SpriteEditorTutorial.SpriteEditorTutorial_C)
++TutorialContexts=(Context="FlipbookEditor",BrowserFilter=,AttractTutorial=None,LaunchTutorial=/Engine/Tutorial/Paper2D/FlipbookEditorTutorial.FlipbookEditorTutorial_C)
++TutorialContexts=(Context="LandscapeMode",BrowserFilter=,AttractTutorial=None,LaunchTutorial=/Engine/Tutorial/Landscape/Landscape_Manage_Mode.Landscape_Manage_Mode_C)
++TutorialContexts=(Context="FoliageMode",BrowserFilter=,AttractTutorial=None,LaunchTutorial=/Engine/Tutorial/Foliage/Foliage_Intro_Tutorial.Foliage_Intro_Tutorial_C)
+
diff --git a/Engine/Config/BaseEngine.ini b/Engine/Config/BaseEngine.ini
--- a/Engine/Config/BaseEngine.ini
+++ b/Engine/Config/BaseEngine.ini
@@ -469,6 +469,10 @@ bUseStreamingPause=false
 +ActiveClassRedirects=(OldClassName="CrowdManager",NewClassName="/Script/AIModule.CrowdManager")
 +ActiveClassRedirects=(OldClassName="PathFollowingComponent",NewClassName="/Script/AIModule.PathFollowingComponent")
 +ActiveClassRedirects=(OldClassName="PawnSensingComponent",NewClassName="/Script/AIModule.PawnSensingComponent")
+
+; We renamed EditorGameAgnosticSettings to be less imposing
++ActiveClassRedirects=(OldClassName="EditorGameAgnosticSettings",NewClassName="/Script/UnrealEd.EditorSettings")
+
 +K2FieldRedirects=(OldFieldName="BTTask_MoveDirectlyToward.bForceMoveToLocation",NewFieldName="bDisablePathUpdateOnGoalLocationChange")
 +K2FieldRedirects=(OldFieldName="BTFunctionLibrary.GetBlackboard", NewFieldName="BTFunctionLibrary.GetOwnersBlackboard")
 
diff --git a/Engine/Plugins/2D/Paper2D/Source/Paper2DEditor/Private/PaperEditorViewportClient.cpp b/Engine/Plugins/2D/Paper2D/Source/Paper2DEditor/Private/PaperEditorViewportClient.cpp
--- a/Engine/Plugins/2D/Paper2D/Source/Paper2DEditor/Private/PaperEditorViewportClient.cpp
+++ b/Engine/Plugins/2D/Paper2D/Source/Paper2DEditor/Private/PaperEditorViewportClient.cpp
@@ -128,9 +128,9 @@ void FPaperEditorViewportClient::ModifyCheckerboardTextureColors()
 
 void FPaperEditorViewportClient::SetupCheckerboardTexture(const FColor& ColorOne, const FColor& ColorTwo, int32 CheckerSize)
 {
-// 	GConfig->GetInt(TEXT("TextureProperties"), TEXT("CheckerboardCheckerPixelNum"), CheckerSize, GEditorUserSettingsIni);
-// 	GConfig->GetColor(TEXT("TextureProperties"), TEXT("CheckerboardColorOne"), CheckerColorOne, GEditorUserSettingsIni);
-// 	GConfig->GetColor(TEXT("TextureProperties"), TEXT("CheckerboardColorTwo"), CheckerColorTwo, GEditorUserSettingsIni);
+// 	GConfig->GetInt(TEXT("TextureProperties"), TEXT("CheckerboardCheckerPixelNum"), CheckerSize, GEditorPerProjectIni);
+// 	GConfig->GetColor(TEXT("TextureProperties"), TEXT("CheckerboardColorOne"), CheckerColorOne, GEditorPerProjectIni);
+// 	GConfig->GetColor(TEXT("TextureProperties"), TEXT("CheckerboardColorTwo"), CheckerColorTwo, GEditorPerProjectIni);
 	
 // 	CheckerColorOne = FColor(128, 128, 128);
 // 	CheckerColorTwo = FColor(64, 64, 64);
diff --git a/Engine/Plugins/Editor/EpicSurvey/Source/EpicSurvey/Private/EpicSurvey.cpp b/Engine/Plugins/Editor/EpicSurvey/Source/EpicSurvey/Private/EpicSurvey.cpp
--- a/Engine/Plugins/Editor/EpicSurvey/Source/EpicSurvey/Private/EpicSurvey.cpp
+++ b/Engine/Plugins/Editor/EpicSurvey/Source/EpicSurvey/Private/EpicSurvey.cpp
@@ -18,6 +18,8 @@
 #include "SNotificationList.h"
 #include "NotificationManager.h"
 
+#include "Settings/EditorSettings.h"
+
 #define LOCTEXT_NAMESPACE "EpicSurvey"
 
 DEFINE_LOG_CATEGORY(LogEpicSurvey);
@@ -310,16 +312,18 @@ void FEpicSurvey::OnReadFileComplete( bool bSuccess, const FString& DLName )
 							{
 							case ESurveyType::Normal:
 								{
+									auto* Settings = GetDefault<UEditorSettings>();
+
 									Surveys.Add( NewSurvey.ToSharedRef() );
-									const bool IsActiveSurveyInProgress = ActiveSurvey.IsValid() ? GEditor->AccessGameAgnosticSettings().InProgressSurveys.Contains( ActiveSurvey->GetIdentifier() ) : false;
+									const bool IsActiveSurveyInProgress = ActiveSurvey.IsValid() ? Settings->InProgressSurveys.Contains( ActiveSurvey->GetIdentifier() ) : false;
 
 									if ( !IsActiveSurveyInProgress )
 									{
-										const bool HasBeenCompleted = GEditor->AccessGameAgnosticSettings().CompletedSurveys.Contains( NewSurvey->GetIdentifier() );
+										const bool HasBeenCompleted = Settings->CompletedSurveys.Contains( NewSurvey->GetIdentifier() );
 
 										if ( !HasBeenCompleted )
 										{
-											const bool IsInProgress = GEditor->AccessGameAgnosticSettings().InProgressSurveys.Contains( NewSurvey->GetIdentifier() );
+											const bool IsInProgress = Settings->InProgressSurveys.Contains( NewSurvey->GetIdentifier() );
 
 											if ( NewSurvey->CanAutoPrompt() )
 											{
@@ -487,7 +491,7 @@ void FEpicSurvey::SetActiveSurvey( const TSharedPtr< FSurvey >& Survey, bool Aut
 
 	if ( Survey.IsValid() )
 	{
-		HasBeenCompleted = GEditor->AccessGameAgnosticSettings().CompletedSurveys.Contains( Survey->GetIdentifier() );
+		HasBeenCompleted = GetDefault<UEditorSettings>()->CompletedSurveys.Contains( Survey->GetIdentifier() );
 
 		if ( !HasBeenCompleted )
 		{
@@ -512,16 +516,17 @@ void FEpicSurvey::SetActiveSurvey( const TSharedPtr< FSurvey >& Survey, bool Aut
 
 void FEpicSurvey::SubmitSurvey( const TSharedRef< FSurvey >& Survey )
 {
-	if ( GEditor->AccessGameAgnosticSettings().CompletedSurveys.Contains( Survey->GetIdentifier() ) )
+	auto* Settings = GetMutableDefault<UEditorSettings>();
+	if ( Settings->CompletedSurveys.Contains( Survey->GetIdentifier() ) )
 	{
 		return;
 	}
 
 	Survey->Submit();
 
-	GEditor->AccessGameAgnosticSettings().CompletedSurveys.Add( Survey->GetIdentifier() );
-	GEditor->AccessGameAgnosticSettings().InProgressSurveys.Remove( Survey->GetIdentifier() );
-	GEditor->AccessGameAgnosticSettings().PostEditChange();
+	Settings->CompletedSurveys.Add( Survey->GetIdentifier() );
+	Settings->InProgressSurveys.Remove( Survey->GetIdentifier() );
+	Settings->PostEditChange();
 
 	SetActiveSurvey( NULL );
 
@@ -683,10 +688,11 @@ void FEpicSurvey::OpenEpicSurveyWindow()
 		{
 			ActiveSurvey->Initialize();
 
-			if ( !GEditor->AccessGameAgnosticSettings().InProgressSurveys.Contains( ActiveSurvey->GetIdentifier() ) )
+			auto* Settings = GetMutableDefault<UEditorSettings>();
+			if ( !Settings->InProgressSurveys.Contains( ActiveSurvey->GetIdentifier() ) )
 			{
-				GEditor->AccessGameAgnosticSettings().InProgressSurveys.Add( ActiveSurvey->GetIdentifier() );
-				GEditor->AccessGameAgnosticSettings().PostEditChange();
+				Settings->InProgressSurveys.Add( ActiveSurvey->GetIdentifier() );
+				Settings->PostEditChange();
 
 				if ( FEngineAnalytics::IsAvailable() )
 				{
diff --git a/Engine/Plugins/Editor/MacGraphicsSwitching/Source/MacGraphicsSwitching/Private/MacGraphicsSwitchingModule.cpp b/Engine/Plugins/Editor/MacGraphicsSwitching/Source/MacGraphicsSwitching/Private/MacGraphicsSwitchingModule.cpp
--- a/Engine/Plugins/Editor/MacGraphicsSwitching/Source/MacGraphicsSwitching/Private/MacGraphicsSwitchingModule.cpp
+++ b/Engine/Plugins/Editor/MacGraphicsSwitching/Source/MacGraphicsSwitching/Private/MacGraphicsSwitchingModule.cpp
@@ -132,7 +132,7 @@ void FMacGraphicsSwitchingModule::Initialize( TSharedPtr<SWindow> InRootWindow,
 		FMacGraphicsSwitchingStyle::Initialize();
 		
 		bool MacUseAutomaticGraphicsSwitching = false;
-		if ( GConfig->GetBool(TEXT("/Script/MacGraphicsSwitching.MacGraphicsSwitchingSettings"), TEXT("bAllowAutomaticGraphicsSwitching"), MacUseAutomaticGraphicsSwitching, GEditorGameAgnosticIni) && MacUseAutomaticGraphicsSwitching )
+		if ( GConfig->GetBool(TEXT("/Script/MacGraphicsSwitching.MacGraphicsSwitchingSettings"), TEXT("bAllowAutomaticGraphicsSwitching"), MacUseAutomaticGraphicsSwitching, GEditorSettingsIni) && MacUseAutomaticGraphicsSwitching )
 		{
 			if ( AllowAutomaticGraphicsSwitching() )
 			{
@@ -148,7 +148,7 @@ void FMacGraphicsSwitchingModule::Initialize( TSharedPtr<SWindow> InRootWindow,
 			}
 			else
 			{
-				GConfig->SetBool(TEXT("/Script/MacGraphicsSwitching.MacGraphicsSwitchingSettings"), TEXT("bAllowAutomaticGraphicsSwitching"), false, GEditorGameAgnosticIni);
+				GConfig->SetBool(TEXT("/Script/MacGraphicsSwitching.MacGraphicsSwitchingSettings"), TEXT("bAllowAutomaticGraphicsSwitching"), false, GEditorSettingsIni);
 			}
 		}
 	}
diff --git a/Engine/Plugins/Editor/MacGraphicsSwitching/Source/MacGraphicsSwitching/Private/MacGraphicsSwitchingSettings.h b/Engine/Plugins/Editor/MacGraphicsSwitching/Source/MacGraphicsSwitching/Private/MacGraphicsSwitchingSettings.h
--- a/Engine/Plugins/Editor/MacGraphicsSwitching/Source/MacGraphicsSwitching/Private/MacGraphicsSwitchingSettings.h
+++ b/Engine/Plugins/Editor/MacGraphicsSwitching/Source/MacGraphicsSwitching/Private/MacGraphicsSwitchingSettings.h
@@ -3,7 +3,7 @@
 
 #include "MacGraphicsSwitchingSettings.generated.h"
 
-UCLASS(Config=EditorGameAgnostic)
+UCLASS(Config=EditorSettings)
 class UMacGraphicsSwitchingSettings : public UObject
 {
 	GENERATED_UCLASS_BODY()
diff --git a/Engine/Plugins/Editor/MacGraphicsSwitching/Source/MacGraphicsSwitching/Private/MacGraphicsSwitchingWidget.cpp b/Engine/Plugins/Editor/MacGraphicsSwitching/Source/MacGraphicsSwitching/Private/MacGraphicsSwitchingWidget.cpp
--- a/Engine/Plugins/Editor/MacGraphicsSwitching/Source/MacGraphicsSwitching/Private/MacGraphicsSwitchingWidget.cpp
+++ b/Engine/Plugins/Editor/MacGraphicsSwitching/Source/MacGraphicsSwitching/Private/MacGraphicsSwitchingWidget.cpp
@@ -133,7 +133,7 @@ void SMacGraphicsSwitchingWidget::OnSelectionChanged(TSharedPtr<FRendererItem> I
 	}
 	else if ( bLiveSwitching )
 	{
-		GConfig->SetInt(TEXT("/Script/MacGraphicsSwitching.MacGraphicsSwitchingSettings"), TEXT("RendererID"), InItem->RendererID, GEditorGameAgnosticIni);
+		GConfig->SetInt(TEXT("/Script/MacGraphicsSwitching.MacGraphicsSwitchingSettings"), TEXT("RendererID"), InItem->RendererID, GEditorSettingsIni);
 	}
 	
 	// Update all the contexts
diff --git a/Engine/Plugins/Experimental/LiveEditor/Source/LiveEditor/Private/LiveEditorManager.cpp b/Engine/Plugins/Experimental/LiveEditor/Source/LiveEditor/Private/LiveEditorManager.cpp
--- a/Engine/Plugins/Experimental/LiveEditor/Source/LiveEditor/Private/LiveEditorManager.cpp
+++ b/Engine/Plugins/Experimental/LiveEditor/Source/LiveEditor/Private/LiveEditorManager.cpp
@@ -246,15 +246,15 @@ void FLiveEditorManager::Initialize()
 	FindDevices();
 	MIDIBuffer = new PmEvent[DEFAULT_BUFFER_SIZE];
 
-	GConfig->GetArray(TEXT("LiveEditor"), TEXT("ActiveBlueprintTemplates"), ActiveBlueprintTemplates, GEditorUserSettingsIni);
+	GConfig->GetArray(TEXT("LiveEditor"), TEXT("ActiveBlueprintTemplates"), ActiveBlueprintTemplates, GEditorPerProjectIni);
 	for ( int32 i = ActiveBlueprintTemplates.Num()-1; i >= 0; --i )
 	{
 		FString &Name = ActiveBlueprintTemplates[i];
 		if ( !FPackageName::DoesPackageExist(*Name) )
 		{
 			ActiveBlueprintTemplates.RemoveAt(i);
-			GConfig->SetArray(TEXT("LiveEditor"), TEXT("ActiveBlueprintTemplates"), ActiveBlueprintTemplates, GEditorUserSettingsIni);
-			GConfig->Flush(false, GEditorUserSettingsIni);
+			GConfig->SetArray(TEXT("LiveEditor"), TEXT("ActiveBlueprintTemplates"), ActiveBlueprintTemplates, GEditorPerProjectIni);
+			GConfig->Flush(false, GEditorPerProjectIni);
 		}
 	}
 
@@ -478,8 +478,8 @@ void FLiveEditorManager::SelectLiveEditorBlueprint( UBlueprint *Blueprint )
 			BlueprintName = BlueprintName.LeftChop(2);
 		UE_LOG( LiveEditorManagerLog, Log, TEXT("%s"), *BlueprintName );
 		ActiveBlueprintTemplates.AddUnique( BlueprintName );
-		GConfig->SetArray(TEXT("LiveEditor"), TEXT("ActiveBlueprintTemplates"), ActiveBlueprintTemplates, GEditorUserSettingsIni);
-		GConfig->Flush(false, GEditorUserSettingsIni);
+		GConfig->SetArray(TEXT("LiveEditor"), TEXT("ActiveBlueprintTemplates"), ActiveBlueprintTemplates, GEditorPerProjectIni);
+		GConfig->Flush(false, GEditorPerProjectIni);
 	}
 }
 
@@ -581,8 +581,8 @@ bool FLiveEditorManager::Remove( const FString &Name )
 	DeActivate( Name );
 	
 	ActiveBlueprintTemplates.Remove( Name );
-	GConfig->SetArray(TEXT("LiveEditor"), TEXT("ActiveBlueprintTemplates"), ActiveBlueprintTemplates, GEditorUserSettingsIni);
-	GConfig->Flush(false, GEditorUserSettingsIni);
+	GConfig->SetArray(TEXT("LiveEditor"), TEXT("ActiveBlueprintTemplates"), ActiveBlueprintTemplates, GEditorPerProjectIni);
+	GConfig->Flush(false, GEditorPerProjectIni);
 	return true;
 }
 
@@ -938,10 +938,10 @@ bool FLiveEditorManager::SaveDeviceData( const FLiveEditorDeviceData &DeviceData
 	}
 
 	FString SectionName = FString::Printf( TEXT("LiveEditor.%s"), *DeviceData.DeviceName );
-	GConfig->SetInt( *SectionName, TEXT("ButtonSignalDown"), DeviceData.ButtonSignalDown, GEditorUserSettingsIni );
-	GConfig->SetInt( *SectionName, TEXT("ButtonSignalUp"), DeviceData.ButtonSignalUp, GEditorUserSettingsIni );
-	GConfig->SetInt( *SectionName, TEXT("ContinuousIncrement"), DeviceData.ContinuousIncrement, GEditorUserSettingsIni );
-	GConfig->SetInt( *SectionName, TEXT("ContinuousDecrement"), DeviceData.ContinuousDecrement, GEditorUserSettingsIni );
+	GConfig->SetInt( *SectionName, TEXT("ButtonSignalDown"), DeviceData.ButtonSignalDown, GEditorPerProjectIni );
+	GConfig->SetInt( *SectionName, TEXT("ButtonSignalUp"), DeviceData.ButtonSignalUp, GEditorPerProjectIni );
+	GConfig->SetInt( *SectionName, TEXT("ContinuousIncrement"), DeviceData.ContinuousIncrement, GEditorPerProjectIni );
+	GConfig->SetInt( *SectionName, TEXT("ContinuousDecrement"), DeviceData.ContinuousDecrement, GEditorPerProjectIni );
 
 	return true;
 }
@@ -950,17 +950,17 @@ bool FLiveEditorManager::LoadDeviceData( const FString &DeviceName, FLiveEditorD
 {
 	FString SectionName = FString::Printf( TEXT("LiveEditor.%s"), *DeviceName );
 	TArray< FString > Section; 
-	if ( !GConfig->GetSection( *SectionName, Section, GEditorUserSettingsIni ) )
+	if ( !GConfig->GetSection( *SectionName, Section, GEditorPerProjectIni ) )
 	{
 		return false;
 	}
 
 	DeviceData.DeviceName = DeviceName;
 	DeviceData.ConfigState = FLiveEditorDeviceData::CONFIGURED;
-	GConfig->GetInt( *SectionName, TEXT("ButtonSignalDown"), DeviceData.ButtonSignalDown, GEditorUserSettingsIni );
-	GConfig->GetInt( *SectionName, TEXT("ButtonSignalUp"), DeviceData.ButtonSignalUp, GEditorUserSettingsIni );
-	GConfig->GetInt( *SectionName, TEXT("ContinuousIncrement"), DeviceData.ContinuousIncrement, GEditorUserSettingsIni );
-	GConfig->GetInt( *SectionName, TEXT("ContinuousDecrement"), DeviceData.ContinuousDecrement, GEditorUserSettingsIni );
+	GConfig->GetInt( *SectionName, TEXT("ButtonSignalDown"), DeviceData.ButtonSignalDown, GEditorPerProjectIni );
+	GConfig->GetInt( *SectionName, TEXT("ButtonSignalUp"), DeviceData.ButtonSignalUp, GEditorPerProjectIni );
+	GConfig->GetInt( *SectionName, TEXT("ContinuousIncrement"), DeviceData.ContinuousIncrement, GEditorPerProjectIni );
+	GConfig->GetInt( *SectionName, TEXT("ContinuousDecrement"), DeviceData.ContinuousDecrement, GEditorPerProjectIni );
 
 	return true;
 }
diff --git a/Engine/Source/Developer/DesktopPlatform/Private/DesktopPlatformBase.cpp b/Engine/Source/Developer/DesktopPlatform/Private/DesktopPlatformBase.cpp
--- a/Engine/Source/Developer/DesktopPlatform/Private/DesktopPlatformBase.cpp
+++ b/Engine/Source/Developer/DesktopPlatform/Private/DesktopPlatformBase.cpp
@@ -893,10 +893,10 @@ bool FDesktopPlatformBase::EnumerateProjectsKnownByEngine(const FString &Identif
 
 	// Load the config file
 	FConfigFile GameAgnosticConfig;
-	FConfigCacheIni::LoadExternalIniFile(GameAgnosticConfig, TEXT("EditorGameAgnostic"), NULL, *GameAgnosticConfigDir, false);
+	FConfigCacheIni::LoadExternalIniFile(GameAgnosticConfig, TEXT("EditorSettings"), NULL, *GameAgnosticConfigDir, false);
 
 	// Find the editor game-agnostic settings
-	FConfigSection* Section = GameAgnosticConfig.Find(TEXT("/Script/UnrealEd.EditorGameAgnosticSettings"));
+	FConfigSection* Section = GameAgnosticConfig.Find(TEXT("/Script/UnrealEd.EditorSettings"));
 	if(Section != NULL)
 	{
 		// Add in every path that the user has ever created a project file. This is to catch new projects showing up in the user's project folders
diff --git a/Engine/Source/Developer/GameplayDebugger/Classes/GameplayDebuggerSettings.h b/Engine/Source/Developer/GameplayDebugger/Classes/GameplayDebuggerSettings.h
--- a/Engine/Source/Developer/GameplayDebugger/Classes/GameplayDebuggerSettings.h
+++ b/Engine/Source/Developer/GameplayDebugger/Classes/GameplayDebuggerSettings.h
@@ -31,7 +31,7 @@ struct FGDTCustomViewNames
 	FString GameView5;
 };
 
-UCLASS(config = EditorUserSettings)
+UCLASS(config = EditorPerProjectUserSettings)
 class GAMEPLAYDEBUGGER_API UGameplayDebuggerSettings : public UObject
 {
 	GENERATED_UCLASS_BODY()
diff --git a/Engine/Source/Developer/GameplayDebugger/Private/GameplayDebuggerSettings.cpp b/Engine/Source/Developer/GameplayDebugger/Private/GameplayDebuggerSettings.cpp
--- a/Engine/Source/Developer/GameplayDebugger/Private/GameplayDebuggerSettings.cpp
+++ b/Engine/Source/Developer/GameplayDebugger/Private/GameplayDebuggerSettings.cpp
@@ -92,7 +92,7 @@ void UGameplayDebuggerSettings::PostEditChangeProperty(FPropertyChangedEvent& Pr
 	UPDATE_GAMEVIEW_DISPLAYNAME(GameView5);
 #undef UPDATE_GAMEVIEW_DISPLAYNAME
 
-	GEditor->SaveEditorUserSettings();
+	SaveConfig();
 
 	SettingChangedEvent.Broadcast(Name);
 }
diff --git a/Engine/Source/Developer/HotReload/Private/HotReload.cpp b/Engine/Source/Developer/HotReload/Private/HotReload.cpp
--- a/Engine/Source/Developer/HotReload/Private/HotReload.cpp
+++ b/Engine/Source/Developer/HotReload/Private/HotReload.cpp
@@ -1527,7 +1527,7 @@ void FHotReloadModule::UpdateModuleCompileData(FName ModuleName)
 void FHotReloadModule::ReadModuleCompilationInfoFromConfig(FName ModuleName, FModuleCompilationData& CompileData)
 {
 	FString DateTimeString;
-	if (GConfig->GetString(*HotReloadDefs::CompilationInfoConfigSection, *FString::Printf(TEXT("%s.TimeStamp"), *ModuleName.ToString()), DateTimeString, GEditorUserSettingsIni))
+	if (GConfig->GetString(*HotReloadDefs::CompilationInfoConfigSection, *FString::Printf(TEXT("%s.TimeStamp"), *ModuleName.ToString()), DateTimeString, GEditorPerProjectIni))
 	{
 		FDateTime TimeStamp;
 		if (!DateTimeString.IsEmpty() && FDateTime::Parse(DateTimeString, TimeStamp))
@@ -1536,7 +1536,7 @@ void FHotReloadModule::ReadModuleCompilationInfoFromConfig(FName ModuleName, FMo
 			CompileData.FileTimeStamp = TimeStamp;
 
 			FString CompileMethodString;
-			if (GConfig->GetString(*HotReloadDefs::CompilationInfoConfigSection, *FString::Printf(TEXT("%s.LastCompileMethod"), *ModuleName.ToString()), CompileMethodString, GEditorUserSettingsIni))
+			if (GConfig->GetString(*HotReloadDefs::CompilationInfoConfigSection, *FString::Printf(TEXT("%s.LastCompileMethod"), *ModuleName.ToString()), CompileMethodString, GEditorPerProjectIni))
 			{
 				if (CompileMethodString.Equals(HotReloadDefs::CompileMethodRuntime, ESearchCase::IgnoreCase))
 				{
@@ -1559,7 +1559,7 @@ void FHotReloadModule::WriteModuleCompilationInfoToConfig(FName ModuleName, cons
 		DateTimeString = CompileData.FileTimeStamp.ToString();
 	}
 
-	GConfig->SetString(*HotReloadDefs::CompilationInfoConfigSection, *FString::Printf(TEXT("%s.TimeStamp"), *ModuleName.ToString()), *DateTimeString, GEditorUserSettingsIni);
+	GConfig->SetString(*HotReloadDefs::CompilationInfoConfigSection, *FString::Printf(TEXT("%s.TimeStamp"), *ModuleName.ToString()), *DateTimeString, GEditorPerProjectIni);
 
 	FString CompileMethodString = HotReloadDefs::CompileMethodUnknown;
 	if (CompileData.CompileMethod == EModuleCompileMethod::Runtime)
@@ -1571,7 +1571,7 @@ void FHotReloadModule::WriteModuleCompilationInfoToConfig(FName ModuleName, cons
 		CompileMethodString = HotReloadDefs::CompileMethodExternal;
 	}
 
-	GConfig->SetString(*HotReloadDefs::CompilationInfoConfigSection, *FString::Printf(TEXT("%s.LastCompileMethod"), *ModuleName.ToString()), *CompileMethodString, GEditorUserSettingsIni);
+	GConfig->SetString(*HotReloadDefs::CompilationInfoConfigSection, *FString::Printf(TEXT("%s.LastCompileMethod"), *ModuleName.ToString()), *CompileMethodString, GEditorPerProjectIni);
 }
 
 bool FHotReloadModule::GetModuleFileTimeStamp(FName ModuleName, FDateTime& OutFileTimeStamp) const
diff --git a/Engine/Source/Developer/InternationalizationSettings/Private/InternationalizationSettingsModel.cpp b/Engine/Source/Developer/InternationalizationSettings/Private/InternationalizationSettingsModel.cpp
--- a/Engine/Source/Developer/InternationalizationSettings/Private/InternationalizationSettingsModel.cpp
+++ b/Engine/Source/Developer/InternationalizationSettings/Private/InternationalizationSettingsModel.cpp
@@ -14,15 +14,15 @@ UInternationalizationSettingsModel::UInternationalizationSettingsModel( const FO
 void UInternationalizationSettingsModel::SaveDefaults()
 {
 	FString SavedCultureName;
-	GConfig->GetString( TEXT("Internationalization"), TEXT("Culture"), SavedCultureName, GEditorGameAgnosticIni );
+	GConfig->GetString( TEXT("Internationalization"), TEXT("Culture"), SavedCultureName, GEditorSettingsIni );
 	GConfig->SetString( TEXT("Internationalization"), TEXT("Culture"), *SavedCultureName, GEngineIni );
 
 	bool bShouldLoadLocalizedPropertyNames = true;
-	GConfig->GetBool( TEXT("Internationalization"), TEXT("ShouldLoadLocalizedPropertyNames"), bShouldLoadLocalizedPropertyNames, GEditorGameAgnosticIni );
+	GConfig->GetBool( TEXT("Internationalization"), TEXT("ShouldLoadLocalizedPropertyNames"), bShouldLoadLocalizedPropertyNames, GEditorSettingsIni );
 	GConfig->SetBool( TEXT("Internationalization"), TEXT("ShouldLoadLocalizedPropertyNames"), bShouldLoadLocalizedPropertyNames, GEngineIni );
 
 	bool bShowNodesAndPinsUnlocalized = false;
-	GConfig->GetBool( TEXT("Internationalization"), TEXT("ShowNodesAndPinsUnlocalized"), bShowNodesAndPinsUnlocalized, GEditorGameAgnosticIni );
+	GConfig->GetBool( TEXT("Internationalization"), TEXT("ShowNodesAndPinsUnlocalized"), bShowNodesAndPinsUnlocalized, GEditorSettingsIni );
 	GConfig->SetBool( TEXT("Internationalization"), TEXT("ShowNodesAndPinsUnlocalized"), bShowNodesAndPinsUnlocalized, GEngineIni );
 }
 
@@ -31,15 +31,15 @@ void UInternationalizationSettingsModel::ResetToDefault()
 {
 	FString SavedCultureName;
 	GConfig->GetString( TEXT("Internationalization"), TEXT("Culture"), SavedCultureName, GEngineIni );
-	GConfig->SetString( TEXT("Internationalization"), TEXT("Culture"), *SavedCultureName, GEditorGameAgnosticIni );
+	GConfig->SetString( TEXT("Internationalization"), TEXT("Culture"), *SavedCultureName, GEditorSettingsIni );
 
 	bool bShouldLoadLocalizedPropertyNames = true;
 	GConfig->GetBool( TEXT("Internationalization"), TEXT("ShouldLoadLocalizedPropertyNames"), bShouldLoadLocalizedPropertyNames, GEngineIni );
-	GConfig->SetBool( TEXT("Internationalization"), TEXT("ShouldLoadLocalizedPropertyNames"), bShouldLoadLocalizedPropertyNames, GEditorGameAgnosticIni );
+	GConfig->SetBool( TEXT("Internationalization"), TEXT("ShouldLoadLocalizedPropertyNames"), bShouldLoadLocalizedPropertyNames, GEditorSettingsIni );
 
 	bool bShowNodesAndPinsUnlocalized = false;
 	GConfig->GetBool( TEXT("Internationalization"), TEXT("ShowNodesAndPinsUnlocalized"), bShowNodesAndPinsUnlocalized, GEngineIni );
-	GConfig->SetBool( TEXT("Internationalization"), TEXT("ShowNodesAndPinsUnlocalized"), bShowNodesAndPinsUnlocalized, GEditorGameAgnosticIni );
+	GConfig->SetBool( TEXT("Internationalization"), TEXT("ShowNodesAndPinsUnlocalized"), bShowNodesAndPinsUnlocalized, GEditorSettingsIni );
 
 	SettingChangedEvent.Broadcast();
 }
@@ -48,7 +48,7 @@ void UInternationalizationSettingsModel::ResetToDefault()
 FString UInternationalizationSettingsModel::GetCultureName() const
 {
 	FString SavedCultureName;
-	if( !GConfig->GetString( TEXT("Internationalization"), TEXT("Culture"), SavedCultureName, GEditorGameAgnosticIni ) )
+	if( !GConfig->GetString( TEXT("Internationalization"), TEXT("Culture"), SavedCultureName, GEditorSettingsIni ) )
 	{
 		GConfig->GetString( TEXT("Internationalization"), TEXT("Culture"), SavedCultureName, GEngineIni );
 	}
@@ -58,7 +58,7 @@ FString UInternationalizationSettingsModel::GetCultureName() const
 
 void UInternationalizationSettingsModel::SetCultureName(const FString& CultureName)
 {
-	GConfig->SetString( TEXT("Internationalization"), TEXT("Culture"), *CultureName, GEditorGameAgnosticIni );
+	GConfig->SetString( TEXT("Internationalization"), TEXT("Culture"), *CultureName, GEditorSettingsIni );
 	SettingChangedEvent.Broadcast();
 }
 
@@ -66,7 +66,7 @@ void UInternationalizationSettingsModel::SetCultureName(const FString& CultureNa
 bool UInternationalizationSettingsModel::ShouldLoadLocalizedPropertyNames() const
 {
 	bool bShouldLoadLocalizedPropertyNames = true;
-	if( !GConfig->GetBool( TEXT("Internationalization"), TEXT("ShouldLoadLocalizedPropertyNames"), bShouldLoadLocalizedPropertyNames, GEditorGameAgnosticIni ) )
+	if( !GConfig->GetBool( TEXT("Internationalization"), TEXT("ShouldLoadLocalizedPropertyNames"), bShouldLoadLocalizedPropertyNames, GEditorSettingsIni ) )
 	{
 		GConfig->GetBool( TEXT("Internationalization"), TEXT("ShouldLoadLocalizedPropertyNames"), bShouldLoadLocalizedPropertyNames, GEngineIni );
 	}
@@ -76,14 +76,14 @@ bool UInternationalizationSettingsModel::ShouldLoadLocalizedPropertyNames() cons
 
 void UInternationalizationSettingsModel::ShouldLoadLocalizedPropertyNames(const bool Value)
 {
-	GConfig->SetBool( TEXT("Internationalization"), TEXT("ShouldLoadLocalizedPropertyNames"), Value, GEditorGameAgnosticIni );
+	GConfig->SetBool( TEXT("Internationalization"), TEXT("ShouldLoadLocalizedPropertyNames"), Value, GEditorSettingsIni );
 	SettingChangedEvent.Broadcast();
 }
 
 bool UInternationalizationSettingsModel::ShouldShowNodesAndPinsUnlocalized() const
 {
 	bool bShowNodesAndPinsUnlocalized = false;
-	if( !GConfig->GetBool( TEXT("Internationalization"), TEXT("ShowNodesAndPinsUnlocalized"), bShowNodesAndPinsUnlocalized, GEditorGameAgnosticIni ) )
+	if( !GConfig->GetBool( TEXT("Internationalization"), TEXT("ShowNodesAndPinsUnlocalized"), bShowNodesAndPinsUnlocalized, GEditorSettingsIni ) )
 	{
 		GConfig->GetBool( TEXT("Internationalization"), TEXT("ShowNodesAndPinsUnlocalized"), bShowNodesAndPinsUnlocalized, GEngineIni );
 	}
@@ -92,5 +92,5 @@ bool UInternationalizationSettingsModel::ShouldShowNodesAndPinsUnlocalized() con
 
 void UInternationalizationSettingsModel::ShouldShowNodesAndPinsUnlocalized(const bool Value)
 {
-	GConfig->SetBool( TEXT("Internationalization"), TEXT("ShowNodesAndPinsUnlocalized"), Value, GEditorGameAgnosticIni );
+	GConfig->SetBool( TEXT("Internationalization"), TEXT("ShowNodesAndPinsUnlocalized"), Value, GEditorSettingsIni );
 }
diff --git a/Engine/Source/Developer/LogVisualizer/Public/LogVisualizerSettings.h b/Engine/Source/Developer/LogVisualizer/Public/LogVisualizerSettings.h
--- a/Engine/Source/Developer/LogVisualizer/Public/LogVisualizerSettings.h
+++ b/Engine/Source/Developer/LogVisualizer/Public/LogVisualizerSettings.h
@@ -43,7 +43,7 @@ struct FVisualLoggerFilters
 
 struct FCategoryFiltersManager;
 
-UCLASS(config = EditorUserSettings)
+UCLASS(config = EditorPerProjectUserSettings)
 class LOGVISUALIZER_API ULogVisualizerSettings : public UObject
 {
 	GENERATED_UCLASS_BODY()
diff --git a/Engine/Source/Developer/MessageLog/Private/Presentation/MessageLogViewModel.cpp b/Engine/Source/Developer/MessageLog/Private/Presentation/MessageLogViewModel.cpp
--- a/Engine/Source/Developer/MessageLog/Private/Presentation/MessageLogViewModel.cpp
+++ b/Engine/Source/Developer/MessageLog/Private/Presentation/MessageLogViewModel.cpp
@@ -41,9 +41,9 @@ void FMessageLogViewModel::ChangeCurrentListingViewModel( const FName& LogName )
 	TSharedPtr<FMessageLogListingViewModel> LogListingViewModel = FindLogListingViewModel( LogName );
 	if( LogListingViewModel.IsValid() )
 	{
-		if ( FPaths::FileExists(GEditorUserSettingsIni) )
+		if ( FPaths::FileExists(GEditorPerProjectIni) )
 		{
-			GConfig->SetString( TEXT("MessageLog"), TEXT("LastLogListing"), *LogName.ToString(), GEditorUserSettingsIni );
+			GConfig->SetString( TEXT("MessageLog"), TEXT("LastLogListing"), *LogName.ToString(), GEditorPerProjectIni );
 		}
 		SelectedLogListingViewModel = LogListingViewModel;
 		SelectionChangedEvent.Broadcast();		
diff --git a/Engine/Source/Developer/MessageLog/Private/UserInterface/SMessageLog.cpp b/Engine/Source/Developer/MessageLog/Private/UserInterface/SMessageLog.cpp
--- a/Engine/Source/Developer/MessageLog/Private/UserInterface/SMessageLog.cpp
+++ b/Engine/Source/Developer/MessageLog/Private/UserInterface/SMessageLog.cpp
@@ -72,11 +72,11 @@ void SMessageLog::Construct( const FArguments& InArgs, const TSharedRef<FMessage
 	{
 		IMessageLogListingPtr DefaultViewModel = ViewModels[0];
 
-		if (FPaths::FileExists(GEditorUserSettingsIni))
+		if (FPaths::FileExists(GEditorPerProjectIni))
 		{
 			FString LogName;
 
-			if (GConfig->GetString(TEXT("MessageLog"), TEXT("LastLogListing"), LogName, GEditorUserSettingsIni))
+			if (GConfig->GetString(TEXT("MessageLog"), TEXT("LastLogListing"), LogName, GEditorPerProjectIni))
 			{
 				IMessageLogListingPtr LoadedViewModel = ViewModel->FindLogListingViewModel(FName(*LogName));
 				
diff --git a/Engine/Source/Developer/SettingsEditor/Private/Widgets/SSettingsEditor.cpp b/Engine/Source/Developer/SettingsEditor/Private/Widgets/SSettingsEditor.cpp
--- a/Engine/Source/Developer/SettingsEditor/Private/Widgets/SSettingsEditor.cpp
+++ b/Engine/Source/Developer/SettingsEditor/Private/Widgets/SSettingsEditor.cpp
@@ -557,7 +557,7 @@ FReply SSettingsEditor::HandleExportButtonClicked()
 	{
 		if (LastExportDir.IsEmpty())
 		{
-			LastExportDir = FPaths::GetPath(GEditorUserSettingsIni);
+			LastExportDir = FPaths::GetPath(GEditorPerProjectIni);
 		}
 
 		FString DefaultFileName = FString::Printf(TEXT("%s Backup %s.ini"), *SelectedSection->GetDisplayName().ToString(), *FDateTime::Now().ToString(TEXT("%Y-%m-%d %H%M%S")));
@@ -607,7 +607,7 @@ FReply SSettingsEditor::HandleImportButtonClicked()
 		TSharedPtr<SWindow> ParentWindow = FSlateApplication::Get().FindWidgetWindow(AsShared());
 		void* ParentWindowHandle = (ParentWindow.IsValid() && ParentWindow->GetNativeWindow().IsValid()) ? ParentWindow->GetNativeWindow()->GetOSWindowHandle() : nullptr;
 
-		if (FDesktopPlatformModule::Get()->OpenFileDialog(ParentWindowHandle, LOCTEXT("ImportSettingsDialogTitle", "Import settings...").ToString(), FPaths::GetPath(GEditorUserSettingsIni), TEXT(""), TEXT("Config files (*.ini)|*.ini"), EFileDialogFlags::None, OutFiles))
+		if (FDesktopPlatformModule::Get()->OpenFileDialog(ParentWindowHandle, LOCTEXT("ImportSettingsDialogTitle", "Import settings...").ToString(), FPaths::GetPath(GEditorPerProjectIni), TEXT(""), TEXT("Config files (*.ini)|*.ini"), EFileDialogFlags::None, OutFiles))
 		{
 			if (SelectedSection->Import(OutFiles[0]) && SelectedSection->Save())
 			{
diff --git a/Engine/Source/Developer/SourceCodeAccess/Private/SourceCodeAccessSettings.h b/Engine/Source/Developer/SourceCodeAccess/Private/SourceCodeAccessSettings.h
--- a/Engine/Source/Developer/SourceCodeAccess/Private/SourceCodeAccessSettings.h
+++ b/Engine/Source/Developer/SourceCodeAccess/Private/SourceCodeAccessSettings.h
@@ -4,7 +4,7 @@
 
 #include "SourceCodeAccessSettings.generated.h"
 
-UCLASS(config=EditorUserSettings)
+UCLASS(config=EditorPerProjectUserSettings)
 class USourceCodeAccessSettings : public UObject
 {
 	GENERATED_UCLASS_BODY()
diff --git a/Engine/Source/Editor/BlueprintGraph/Private/BlueprintEditorSettings.cpp b/Engine/Source/Editor/BlueprintGraph/Private/BlueprintEditorSettings.cpp
--- a/Engine/Source/Editor/BlueprintGraph/Private/BlueprintEditorSettings.cpp
+++ b/Engine/Source/Editor/BlueprintGraph/Private/BlueprintEditorSettings.cpp
@@ -3,7 +3,7 @@
 #include "BlueprintGraphPrivatePCH.h"
 #include "BlueprintEditorSettings.h"
 #include "Editor/UnrealEd/Classes/Settings/EditorExperimentalSettings.h"
-#include "Editor/UnrealEd/Classes/Editor/EditorUserSettings.h"
+#include "Editor/UnrealEd/Classes/Editor/EditorPerProjectUserSettings.h"
 
 UBlueprintEditorSettings::UBlueprintEditorSettings(const FObjectInitializer& ObjectInitializer)
 	: Super(ObjectInitializer)
@@ -38,14 +38,14 @@ UBlueprintEditorSettings::UBlueprintEditorSettings(const FObjectInitializer& Obj
 	bDrawMidpointArrowsInBlueprints = ExperimentalSettings->bDrawMidpointArrowsInBlueprints;
 
 	// settings that were moved out of editor-user settings...
-	UEditorUserSettings const* UserSettings = GetDefault<UEditorUserSettings>();
+	UEditorPerProjectUserSettings const* UserSettings = GetDefault<UEditorPerProjectUserSettings>();
 	bShowActionMenuItemSignatures = UserSettings->bDisplayActionListItemRefIds;
 
 	FString const ClassConfigKey = GetClass()->GetPathName();
 
 	bool bOldSaveOnCompileVal = false;
 	// backwards compatibility: handle the case where users have already switched this on
-	if (GConfig->GetBool(*ClassConfigKey, TEXT("bSaveOnCompile"), bOldSaveOnCompileVal, GEditorUserSettingsIni) && bOldSaveOnCompileVal)
+	if (GConfig->GetBool(*ClassConfigKey, TEXT("bSaveOnCompile"), bOldSaveOnCompileVal, GEditorPerProjectIni) && bOldSaveOnCompileVal)
 	{
 		SaveOnCompile = SoC_SuccessOnly;
 	}
diff --git a/Engine/Source/Editor/BlueprintGraph/Private/K2Node_ClassDynamicCast.cpp b/Engine/Source/Editor/BlueprintGraph/Private/K2Node_ClassDynamicCast.cpp
--- a/Engine/Source/Editor/BlueprintGraph/Private/K2Node_ClassDynamicCast.cpp
+++ b/Engine/Source/Editor/BlueprintGraph/Private/K2Node_ClassDynamicCast.cpp
@@ -67,7 +67,6 @@ void UK2Node_ClassDynamicCast::AllocateDefaultPins()
 
 FLinearColor UK2Node_ClassDynamicCast::GetNodeTitleColor() const
 {
-	const UEditorUserSettings& Options = GEditor->AccessEditorUserSettings();
 	return GetDefault<UGraphEditorSettings>()->ClassPinTypeColor;
 }
 
diff --git a/Engine/Source/Editor/BlueprintGraph/Public/BlueprintEditorSettings.h b/Engine/Source/Editor/BlueprintGraph/Public/BlueprintEditorSettings.h
--- a/Engine/Source/Editor/BlueprintGraph/Public/BlueprintEditorSettings.h
+++ b/Engine/Source/Editor/BlueprintGraph/Public/BlueprintEditorSettings.h
@@ -13,7 +13,7 @@ enum ESaveOnCompile
 	SoC_Always UMETA(DisplayName = "Always"),
 };
  
-UCLASS(config=EditorUserSettings)
+UCLASS(config=EditorPerProjectUserSettings)
 class BLUEPRINTGRAPH_API UBlueprintEditorSettings
 	:	public UObject
 {
diff --git a/Engine/Source/Editor/Blutility/Classes/GlobalEditorUtilityBase.h b/Engine/Source/Editor/Blutility/Classes/GlobalEditorUtilityBase.h
--- a/Engine/Source/Editor/Blutility/Classes/GlobalEditorUtilityBase.h
+++ b/Engine/Source/Editor/Blutility/Classes/GlobalEditorUtilityBase.h
@@ -68,7 +68,7 @@ class BLUTILITY_API UGlobalEditorUtilityBase : public UObject
 	///////////////////////////
 
 	UFUNCTION(BlueprintPure, Category="Development|Editor")
-	UEditorUserSettings* GetEditorUserSettings();
+	UEditorPerProjectUserSettings* GetEditorUserSettings();
 
 	// Remove all actors from the selection set
 	UFUNCTION(BlueprintCallable, Category="Development|Editor")
diff --git a/Engine/Source/Editor/Blutility/Private/GlobalEditorUtilityBase.cpp b/Engine/Source/Editor/Blutility/Private/GlobalEditorUtilityBase.cpp
--- a/Engine/Source/Editor/Blutility/Private/GlobalEditorUtilityBase.cpp
+++ b/Engine/Source/Editor/Blutility/Private/GlobalEditorUtilityBase.cpp
@@ -95,9 +95,9 @@ void UGlobalEditorUtilityBase::ForEachSelectedAsset()
 	}
 }
 
-UEditorUserSettings* UGlobalEditorUtilityBase::GetEditorUserSettings()
+UEditorPerProjectUserSettings* UGlobalEditorUtilityBase::GetEditorUserSettings()
 {
-	return &(GEditor->AccessEditorUserSettings());
+	return GetMutableDefault<UEditorPerProjectUserSettings>();
 }
 
 void UGlobalEditorUtilityBase::ClearActorSelectionSet()
diff --git a/Engine/Source/Editor/ContentBrowser/Private/CollectionViewUtils.cpp b/Engine/Source/Editor/ContentBrowser/Private/CollectionViewUtils.cpp
--- a/Engine/Source/Editor/ContentBrowser/Private/CollectionViewUtils.cpp
+++ b/Engine/Source/Editor/ContentBrowser/Private/CollectionViewUtils.cpp
@@ -91,11 +91,11 @@ const TSharedPtr<FLinearColor> CollectionViewUtils::LoadColor(const FString& InC
 	}
 		
 	// Loads the color of collection at the given path from the config
-	if(FPaths::FileExists(GEditorUserSettingsIni))
+	if(FPaths::FileExists(GEditorPerProjectIni))
 	{
 		// Create a new entry from the config, skip if it's default
 		FString ColorStr;
-		if(GConfig->GetString(TEXT("CollectionColor"), *ColorKeyStr, ColorStr, GEditorUserSettingsIni))
+		if(GConfig->GetString(TEXT("CollectionColor"), *ColorKeyStr, ColorStr, GEditorPerProjectIni))
 		{
 			FLinearColor Color;
 			if(Color.InitFromString(ColorStr) && !Color.Equals(CollectionViewUtils::GetDefaultColor()))
@@ -122,16 +122,16 @@ void CollectionViewUtils::SaveColor(const FString& InCollectionName, const EColl
 	const bool bRemove = !CollectionColor.IsValid() || (!bForceAdd && CollectionColor->Equals(CollectionViewUtils::GetDefaultColor()));
 
 	// Saves the color of the collection to the config
-	if(FPaths::FileExists(GEditorUserSettingsIni))
+	if(FPaths::FileExists(GEditorPerProjectIni))
 	{
 		// If this is no longer custom, remove it
 		if(bRemove)
 		{
-			GConfig->RemoveKey(TEXT("CollectionColor"), *ColorKeyStr, GEditorUserSettingsIni);
+			GConfig->RemoveKey(TEXT("CollectionColor"), *ColorKeyStr, GEditorPerProjectIni);
 		}
 		else
 		{
-			GConfig->SetString(TEXT("CollectionColor"), *ColorKeyStr, *CollectionColor->ToString(), GEditorUserSettingsIni);
+			GConfig->SetString(TEXT("CollectionColor"), *ColorKeyStr, *CollectionColor->ToString(), GEditorPerProjectIni);
 		}
 	}
 
@@ -148,14 +148,14 @@ void CollectionViewUtils::SaveColor(const FString& InCollectionName, const EColl
 
 bool CollectionViewUtils::HasCustomColors( TArray< FLinearColor >* OutColors )
 {
-	if(!FPaths::FileExists(GEditorUserSettingsIni))
+	if(!FPaths::FileExists(GEditorPerProjectIni))
 	{
 		return false;
 	}
 
 	// Read individual entries from a config file.
 	TArray<FString> Section;
-	GConfig->GetSection(TEXT("CollectionColor"), Section, GEditorUserSettingsIni);
+	GConfig->GetSection(TEXT("CollectionColor"), Section, GEditorPerProjectIni);
 
 	bool bHasCustom = false;
 	const FCollectionManagerModule& CollectionManagerModule = FModuleManager::Get().LoadModuleChecked<FCollectionManagerModule>("CollectionManager");
diff --git a/Engine/Source/Editor/ContentBrowser/Private/ContentBrowserSingleton.cpp b/Engine/Source/Editor/ContentBrowser/Private/ContentBrowserSingleton.cpp
--- a/Engine/Source/Editor/ContentBrowser/Private/ContentBrowserSingleton.cpp
+++ b/Engine/Source/Editor/ContentBrowser/Private/ContentBrowserSingleton.cpp
@@ -509,7 +509,7 @@ bool FContentBrowserSingleton::IsLocked(const FName& InstanceName) const
 
 	// Fallback to getting the locked state from the config instead
 	bool bIsLocked = false;
-	GConfig->GetBool(*SContentBrowser::SettingsIniSection, *(InstanceName.ToString() + TEXT(".Locked")), bIsLocked, GEditorUserSettingsIni);
+	GConfig->GetBool(*SContentBrowser::SettingsIniSection, *(InstanceName.ToString() + TEXT(".Locked")), bIsLocked, GEditorPerProjectIni);
 	return bIsLocked;
 }
 
diff --git a/Engine/Source/Editor/ContentBrowser/Private/ContentBrowserUtils.cpp b/Engine/Source/Editor/ContentBrowser/Private/ContentBrowserUtils.cpp
--- a/Engine/Source/Editor/ContentBrowser/Private/ContentBrowserUtils.cpp
+++ b/Engine/Source/Editor/ContentBrowser/Private/ContentBrowserUtils.cpp
@@ -1334,11 +1334,11 @@ const TSharedPtr<FLinearColor> ContentBrowserUtils::LoadColor(const FString& Fol
 		}
 		
 		// Loads the color of folder at the given path from the config
-		if(FPaths::FileExists(GEditorUserSettingsIni))
+		if(FPaths::FileExists(GEditorPerProjectIni))
 		{
 			// Create a new entry from the config, skip if it's default
 			FString ColorStr;
-			if(GConfig->GetString(TEXT("PathColor"), *InPath, ColorStr, GEditorUserSettingsIni))
+			if(GConfig->GetString(TEXT("PathColor"), *InPath, ColorStr, GEditorPerProjectIni))
 			{
 				FLinearColor Color;
 				if(Color.InitFromString(ColorStr) && !Color.Equals(ContentBrowserUtils::GetDefaultColor()))
@@ -1377,9 +1377,9 @@ void ContentBrowserUtils::SaveColor(const FString& FolderPath, const TSharedPtr<
 	auto SaveColorInternal = [](const FString& InPath, const TSharedPtr<FLinearColor>& InFolderColor)
 	{
 		// Saves the color of the folder to the config
-		if(FPaths::FileExists(GEditorUserSettingsIni))
+		if(FPaths::FileExists(GEditorPerProjectIni))
 		{
-			GConfig->SetString(TEXT("PathColor"), *InPath, *InFolderColor->ToString(), GEditorUserSettingsIni);
+			GConfig->SetString(TEXT("PathColor"), *InPath, *InFolderColor->ToString(), GEditorPerProjectIni);
 		}
 
 		// Update the map too
@@ -1389,9 +1389,9 @@ void ContentBrowserUtils::SaveColor(const FString& FolderPath, const TSharedPtr<
 	auto RemoveColorInternal = [](const FString& InPath)
 	{
 		// Remove the color of the folder from the config
-		if(FPaths::FileExists(GEditorUserSettingsIni))
+		if(FPaths::FileExists(GEditorPerProjectIni))
 		{
-			GConfig->RemoveKey(TEXT("PathColor"), *InPath, GEditorUserSettingsIni);
+			GConfig->RemoveKey(TEXT("PathColor"), *InPath, GEditorPerProjectIni);
 		}
 
 		// Update the map too
@@ -1423,11 +1423,11 @@ bool ContentBrowserUtils::HasCustomColors( TArray< FLinearColor >* OutColors )
 	// Check to see how many paths are currently using this color
 	// Note: we have to use the config, as paths which haven't been rendered yet aren't registered in the map
 	bool bHasCustom = false;
-	if (FPaths::FileExists(GEditorUserSettingsIni))
+	if (FPaths::FileExists(GEditorPerProjectIni))
 	{
 		// Read individual entries from a config file.
 		TArray< FString > Section; 
-		GConfig->GetSection( TEXT("PathColor"), Section, GEditorUserSettingsIni );
+		GConfig->GetSection( TEXT("PathColor"), Section, GEditorPerProjectIni );
 
 		for( int32 SectionIndex = 0; SectionIndex < Section.Num(); SectionIndex++ )
 		{
diff --git a/Engine/Source/Editor/ContentBrowser/Private/SAssetPicker.cpp b/Engine/Source/Editor/ContentBrowser/Private/SAssetPicker.cpp
--- a/Engine/Source/Editor/ContentBrowser/Private/SAssetPicker.cpp
+++ b/Engine/Source/Editor/ContentBrowser/Private/SAssetPicker.cpp
@@ -480,10 +480,10 @@ void SAssetPicker::LoadSettings()
 		// Load all our data using the settings string as a key in the user settings ini
 		if (FilterListPtr.IsValid())
 		{
-			FilterListPtr->LoadSettings(GEditorUserSettingsIni, SContentBrowser::SettingsIniSection, SettingsString);
+			FilterListPtr->LoadSettings(GEditorPerProjectIni, SContentBrowser::SettingsIniSection, SettingsString);
 		}
 		
-		AssetViewPtr->LoadSettings(GEditorUserSettingsIni, SContentBrowser::SettingsIniSection, SettingsString);
+		AssetViewPtr->LoadSettings(GEditorPerProjectIni, SContentBrowser::SettingsIniSection, SettingsString);
 	}
 }
 
@@ -496,10 +496,10 @@ void SAssetPicker::SaveSettings() const
 		// Save all our data using the settings string as a key in the user settings ini
 		if (FilterListPtr.IsValid())
 		{
-			FilterListPtr->SaveSettings(GEditorUserSettingsIni, SContentBrowser::SettingsIniSection, SettingsString);
+			FilterListPtr->SaveSettings(GEditorPerProjectIni, SContentBrowser::SettingsIniSection, SettingsString);
 		}
 
-		AssetViewPtr->SaveSettings(GEditorUserSettingsIni, SContentBrowser::SettingsIniSection, SettingsString);
+		AssetViewPtr->SaveSettings(GEditorPerProjectIni, SContentBrowser::SettingsIniSection, SettingsString);
 	}
 }
 
diff --git a/Engine/Source/Editor/ContentBrowser/Private/SContentBrowser.cpp b/Engine/Source/Editor/ContentBrowser/Private/SContentBrowser.cpp
--- a/Engine/Source/Editor/ContentBrowser/Private/SContentBrowser.cpp
+++ b/Engine/Source/Editor/ContentBrowser/Private/SContentBrowser.cpp
@@ -895,26 +895,26 @@ void SContentBrowser::SaveSettings() const
 {
 	const FString& SettingsString = InstanceName.ToString();
 
-	GConfig->SetBool(*SettingsIniSection, *(SettingsString + TEXT(".SourcesExpanded")), bSourcesViewExpanded, GEditorUserSettingsIni);
-	GConfig->SetBool(*SettingsIniSection, *(SettingsString + TEXT(".Locked")), bIsLocked, GEditorUserSettingsIni);
+	GConfig->SetBool(*SettingsIniSection, *(SettingsString + TEXT(".SourcesExpanded")), bSourcesViewExpanded, GEditorPerProjectIni);
+	GConfig->SetBool(*SettingsIniSection, *(SettingsString + TEXT(".Locked")), bIsLocked, GEditorPerProjectIni);
 
 	for(int32 SlotIndex = 0; SlotIndex < PathAssetSplitterPtr->GetChildren()->Num(); SlotIndex++)
 	{
 		float SplitterSize = PathAssetSplitterPtr->SlotAt(SlotIndex).SizeValue.Get();
-		GConfig->SetFloat(*SettingsIniSection, *(SettingsString + FString::Printf(TEXT(".VerticalSplitter.SlotSize%d"), SlotIndex)), SplitterSize, GEditorUserSettingsIni);
+		GConfig->SetFloat(*SettingsIniSection, *(SettingsString + FString::Printf(TEXT(".VerticalSplitter.SlotSize%d"), SlotIndex)), SplitterSize, GEditorPerProjectIni);
 	}
 	
 	for(int32 SlotIndex = 0; SlotIndex < PathCollectionSplitterPtr->GetChildren()->Num(); SlotIndex++)
 	{
 		float SplitterSize = PathCollectionSplitterPtr->SlotAt(SlotIndex).SizeValue.Get();
-		GConfig->SetFloat(*SettingsIniSection, *(SettingsString + FString::Printf(TEXT(".HorizontalSplitter.SlotSize%d"), SlotIndex)), SplitterSize, GEditorUserSettingsIni);
+		GConfig->SetFloat(*SettingsIniSection, *(SettingsString + FString::Printf(TEXT(".HorizontalSplitter.SlotSize%d"), SlotIndex)), SplitterSize, GEditorPerProjectIni);
 	}
 
 	// Save all our data using the settings string as a key in the user settings ini
-	FilterListPtr->SaveSettings(GEditorUserSettingsIni, SettingsIniSection, SettingsString);
-	PathViewPtr->SaveSettings(GEditorUserSettingsIni, SettingsIniSection, SettingsString);
-	CollectionViewPtr->SaveSettings(GEditorUserSettingsIni, SettingsIniSection, SettingsString);
-	AssetViewPtr->SaveSettings(GEditorUserSettingsIni, SettingsIniSection, SettingsString);
+	FilterListPtr->SaveSettings(GEditorPerProjectIni, SettingsIniSection, SettingsString);
+	PathViewPtr->SaveSettings(GEditorPerProjectIni, SettingsIniSection, SettingsString);
+	CollectionViewPtr->SaveSettings(GEditorPerProjectIni, SettingsIniSection, SettingsString);
+	AssetViewPtr->SaveSettings(GEditorPerProjectIni, SettingsIniSection, SettingsString);
 }
 
 const FName SContentBrowser::GetInstanceName() const
@@ -1009,16 +1009,16 @@ void SContentBrowser::LoadSettings(const FName& InInstanceName)
 	// Test to see if we should load legacy settings from a previous instance name
 	// First make sure there aren't any existing settings with the given instance name
 	bool TestBool;
-	if ( !GConfig->GetBool(*SettingsIniSection, *(SettingsString + TEXT(".SourcesExpanded")), TestBool, GEditorUserSettingsIni) )
+	if ( !GConfig->GetBool(*SettingsIniSection, *(SettingsString + TEXT(".SourcesExpanded")), TestBool, GEditorPerProjectIni) )
 	{
 		// If there were not any settings and we are Content Browser 1, see if we have any settings under the legacy name "LevelEditorContentBrowser"
-		if ( InInstanceName.ToString() == TEXT("ContentBrowserTab1") && GConfig->GetBool(*SettingsIniSection, TEXT("LevelEditorContentBrowser.SourcesExpanded"), TestBool, GEditorUserSettingsIni) )
+		if ( InInstanceName.ToString() == TEXT("ContentBrowserTab1") && GConfig->GetBool(*SettingsIniSection, TEXT("LevelEditorContentBrowser.SourcesExpanded"), TestBool, GEditorPerProjectIni) )
 		{
 			// We have found some legacy settings with the old ID, use them. These settings will be saved out to the new id later
 			SettingsString = TEXT("LevelEditorContentBrowser");
 		}
 		// else see if we are Content Browser 2, and see if we have any settings under the legacy name "MajorContentBrowserTab"
-		else if ( InInstanceName.ToString() == TEXT("ContentBrowserTab2") && GConfig->GetBool(*SettingsIniSection, TEXT("MajorContentBrowserTab.SourcesExpanded"), TestBool, GEditorUserSettingsIni) )
+		else if ( InInstanceName.ToString() == TEXT("ContentBrowserTab2") && GConfig->GetBool(*SettingsIniSection, TEXT("MajorContentBrowserTab.SourcesExpanded"), TestBool, GEditorPerProjectIni) )
 		{
 			// We have found some legacy settings with the old ID, use them. These settings will be saved out to the new id later
 			SettingsString = TEXT("MajorContentBrowserTab");
@@ -1026,28 +1026,28 @@ void SContentBrowser::LoadSettings(const FName& InInstanceName)
 	}
 
 	// Now that we have determined the appropriate settings string, actually load the settings
-	GConfig->GetBool(*SettingsIniSection, *(SettingsString + TEXT(".SourcesExpanded")), bSourcesViewExpanded, GEditorUserSettingsIni);
-	GConfig->GetBool(*SettingsIniSection, *(SettingsString + TEXT(".Locked")), bIsLocked, GEditorUserSettingsIni);
+	GConfig->GetBool(*SettingsIniSection, *(SettingsString + TEXT(".SourcesExpanded")), bSourcesViewExpanded, GEditorPerProjectIni);
+	GConfig->GetBool(*SettingsIniSection, *(SettingsString + TEXT(".Locked")), bIsLocked, GEditorPerProjectIni);
 
 	for(int32 SlotIndex = 0; SlotIndex < PathAssetSplitterPtr->GetChildren()->Num(); SlotIndex++)
 	{
 		float SplitterSize = PathAssetSplitterPtr->SlotAt(SlotIndex).SizeValue.Get();
-		GConfig->GetFloat(*SettingsIniSection, *(SettingsString + FString::Printf(TEXT(".VerticalSplitter.SlotSize%d"), SlotIndex)), SplitterSize, GEditorUserSettingsIni);
+		GConfig->GetFloat(*SettingsIniSection, *(SettingsString + FString::Printf(TEXT(".VerticalSplitter.SlotSize%d"), SlotIndex)), SplitterSize, GEditorPerProjectIni);
 		PathAssetSplitterPtr->SlotAt(SlotIndex).SizeValue = SplitterSize;
 	}
 	
 	for(int32 SlotIndex = 0; SlotIndex < PathCollectionSplitterPtr->GetChildren()->Num(); SlotIndex++)
 	{
 		float SplitterSize = PathCollectionSplitterPtr->SlotAt(SlotIndex).SizeValue.Get();
-		GConfig->GetFloat(*SettingsIniSection, *(SettingsString + FString::Printf(TEXT(".HorizontalSplitter.SlotSize%d"), SlotIndex)), SplitterSize, GEditorUserSettingsIni);
+		GConfig->GetFloat(*SettingsIniSection, *(SettingsString + FString::Printf(TEXT(".HorizontalSplitter.SlotSize%d"), SlotIndex)), SplitterSize, GEditorPerProjectIni);
 		PathCollectionSplitterPtr->SlotAt(SlotIndex).SizeValue = SplitterSize;
 	}
 
 	// Save all our data using the settings string as a key in the user settings ini
-	FilterListPtr->LoadSettings(GEditorUserSettingsIni, SettingsIniSection, SettingsString);
-	PathViewPtr->LoadSettings(GEditorUserSettingsIni, SettingsIniSection, SettingsString);
-	CollectionViewPtr->LoadSettings(GEditorUserSettingsIni, SettingsIniSection, SettingsString);
-	AssetViewPtr->LoadSettings(GEditorUserSettingsIni, SettingsIniSection, SettingsString);
+	FilterListPtr->LoadSettings(GEditorPerProjectIni, SettingsIniSection, SettingsString);
+	PathViewPtr->LoadSettings(GEditorPerProjectIni, SettingsIniSection, SettingsString);
+	CollectionViewPtr->LoadSettings(GEditorPerProjectIni, SettingsIniSection, SettingsString);
+	AssetViewPtr->LoadSettings(GEditorPerProjectIni, SettingsIniSection, SettingsString);
 }
 
 void SContentBrowser::SourcesChanged(const TArray<FString>& SelectedPaths, const TArray<FCollectionNameType>& SelectedCollections)
diff --git a/Engine/Source/Editor/ContentBrowser/Private/SContentBrowser.h b/Engine/Source/Editor/ContentBrowser/Private/SContentBrowser.h
--- a/Engine/Source/Editor/ContentBrowser/Private/SContentBrowser.h
+++ b/Engine/Source/Editor/ContentBrowser/Private/SContentBrowser.h
@@ -398,6 +398,6 @@ private:
 
 public: 
 
-	/** The section of EditorUserSettings in which to save content browser settings */
+	/** The section of EditorPerProjectUserSettings in which to save content browser settings */
 	static const FString SettingsIniSection;
 };
diff --git a/Engine/Source/Editor/DetailCustomizations/Private/ComponentTransformDetails.cpp b/Engine/Source/Editor/DetailCustomizations/Private/ComponentTransformDetails.cpp
--- a/Engine/Source/Editor/DetailCustomizations/Private/ComponentTransformDetails.cpp
+++ b/Engine/Source/Editor/DetailCustomizations/Private/ComponentTransformDetails.cpp
@@ -66,7 +66,7 @@ FComponentTransformDetails::FComponentTransformDetails( const TArray< TWeakObjec
 	, bEditingRotationInUI( false )
 	, HiddenFieldMask( 0 )
 {
-	GConfig->GetBool(TEXT("SelectionDetails"), TEXT("PreserveScaleRatio"), bPreserveScaleRatio, GEditorUserSettingsIni);
+	GConfig->GetBool(TEXT("SelectionDetails"), TEXT("PreserveScaleRatio"), bPreserveScaleRatio, GEditorPerProjectIni);
 
 }
 
@@ -636,7 +636,7 @@ ECheckBoxState FComponentTransformDetails::IsPreserveScaleRatioChecked() const
 void FComponentTransformDetails::OnPreserveScaleRatioToggled( ECheckBoxState NewState )
 {
 	bPreserveScaleRatio = (NewState == ECheckBoxState::Checked) ? true : false;
-	GConfig->SetBool(TEXT("SelectionDetails"), TEXT("PreserveScaleRatio"), bPreserveScaleRatio, GEditorUserSettingsIni);
+	GConfig->SetBool(TEXT("SelectionDetails"), TEXT("PreserveScaleRatio"), bPreserveScaleRatio, GEditorPerProjectIni);
 }
 
 FText FComponentTransformDetails::GetLocationText() const
diff --git a/Engine/Source/Editor/DetailCustomizations/Private/MathStructCustomizations.cpp b/Engine/Source/Editor/DetailCustomizations/Private/MathStructCustomizations.cpp
--- a/Engine/Source/Editor/DetailCustomizations/Private/MathStructCustomizations.cpp
+++ b/Engine/Source/Editor/DetailCustomizations/Private/MathStructCustomizations.cpp
@@ -74,7 +74,7 @@ void FMathStructCustomization::MakeHeaderRow( TSharedRef<class IPropertyHandle>&
 
 	if ( StructPropertyHandle->GetProperty()->HasMetaData("AllowPreserveRatio") )
 	{
-		if ( !GConfig->GetBool(TEXT("SelectionDetails"), *(StructPropertyHandle->GetProperty()->GetName() + TEXT("_PreserveScaleRatio")), bPreserveScaleRatio, GEditorUserSettingsIni) )
+		if ( !GConfig->GetBool(TEXT("SelectionDetails"), *(StructPropertyHandle->GetProperty()->GetName() + TEXT("_PreserveScaleRatio")), bPreserveScaleRatio, GEditorPerProjectIni) )
 		{
 			bPreserveScaleRatio = true;
 		}
@@ -115,7 +115,7 @@ void FMathStructCustomization::OnPreserveScaleRatioToggled(ECheckBoxState NewSta
 	if ( PropertyHandle.IsValid() )
 	{
 		FString SettingKey = ( PropertyHandle.Pin()->GetProperty()->GetName() + TEXT("_PreserveScaleRatio") );
-		GConfig->SetBool(TEXT("SelectionDetails"), *SettingKey, bPreserveScaleRatio, GEditorUserSettingsIni);
+		GConfig->SetBool(TEXT("SelectionDetails"), *SettingKey, bPreserveScaleRatio, GEditorPerProjectIni);
 	}
 }
 
diff --git a/Engine/Source/Editor/DeviceProfileServices/Private/DeviceProfileServicesUIManager.cpp b/Engine/Source/Editor/DeviceProfileServices/Private/DeviceProfileServicesUIManager.cpp
--- a/Engine/Source/Editor/DeviceProfileServices/Private/DeviceProfileServicesUIManager.cpp
+++ b/Engine/Source/Editor/DeviceProfileServices/Private/DeviceProfileServicesUIManager.cpp
@@ -102,7 +102,7 @@ void FDeviceProfileServicesUIManager::SetProfile( const FString& DeviceProfileNa
 		// Get the existing items
 		for( int32 ItemIdx = 0 ; ItemIdx < MaxItems; ++ItemIdx )
 		{
-			if ( GConfig->GetString( *INISection, *FString::Printf( TEXT("%s%d"), *INIKeyBase, ItemIdx ), CurItem, GEditorUserSettingsIni ) )
+			if ( GConfig->GetString( *INISection, *FString::Printf( TEXT("%s%d"), *INIKeyBase, ItemIdx ), CurItem, GEditorPerProjectIni ) )
 			{
 				CurItems.Add( CurItem );
 			}
@@ -124,14 +124,14 @@ void FDeviceProfileServicesUIManager::SetProfile( const FString& DeviceProfileNa
 		CurItems.Insert( DeviceProfileName, 0 );
 
 		// Clear the ini section
-		GConfig->EmptySection( *INISection, GEditorUserSettingsIni );
+		GConfig->EmptySection( *INISection, GEditorPerProjectIni );
 
 		// Re-write the .ini file
 		for ( int32 ItemIdx = 0; ItemIdx < CurItems.Num(); ++ItemIdx )
 		{
-			GConfig->SetString( *INISection, *FString::Printf( TEXT("%s%d"), *INIKeyBase, ItemIdx ), *CurItems[ItemIdx], GEditorUserSettingsIni );
+			GConfig->SetString( *INISection, *FString::Printf( TEXT("%s%d"), *INIKeyBase, ItemIdx ), *CurItems[ItemIdx], GEditorPerProjectIni );
 		}
 
-		GConfig->Flush( false, GEditorUserSettingsIni );
+		GConfig->Flush( false, GEditorPerProjectIni );
 	}
 }
diff --git a/Engine/Source/Editor/Documentation/Private/SDocumentationToolTip.cpp b/Engine/Source/Editor/Documentation/Private/SDocumentationToolTip.cpp
--- a/Engine/Source/Editor/Documentation/Private/SDocumentationToolTip.cpp
+++ b/Engine/Source/Editor/Documentation/Private/SDocumentationToolTip.cpp
@@ -93,9 +93,9 @@ void SDocumentationToolTip::AddDocumentation(TSharedPtr< SVerticalBox > Vertical
 {
 	if ( !DocumentationLink.IsEmpty() )
 	{
-		IsDisplayingDocumentationLink = GEditor->EditorUserSettings->bDisplayDocumentationLink;
+		IsDisplayingDocumentationLink = GetDefault<UEditorPerProjectUserSettings>()->bDisplayDocumentationLink;
 
-		if ( GEditor->EditorUserSettings->bDisplayDocumentationLink )
+		if ( IsDisplayingDocumentationLink )
 		{
 			FString OptionalExcerptName;
 			if ( !ExcerptName.IsEmpty() )
@@ -143,7 +143,7 @@ void SDocumentationToolTip::AddDocumentation(TSharedPtr< SVerticalBox > Vertical
 		}
 		else
 		{
-			if ( GEditor->EditorUserSettings->bDisplayDocumentationLink && FSlateApplication::Get().SupportsSourceAccess() )
+			if ( IsDisplayingDocumentationLink && FSlateApplication::Get().SupportsSourceAccess() )
 			{
 				FString DocPath = FDocumentationLink::ToSourcePath( DocumentationLink, FInternationalization::Get().GetCurrentCulture() );
 				if ( !FPaths::FileExists(DocPath) )
@@ -308,7 +308,7 @@ void SDocumentationToolTip::ConstructFullTipContent()
 				];
 			}
 
-			if ( GEditor->EditorUserSettings->bDisplayDocumentationLink && FSlateApplication::Get().SupportsSourceAccess() )
+			if ( GetDefault<UEditorPerProjectUserSettings>()->bDisplayDocumentationLink && FSlateApplication::Get().SupportsSourceAccess() )
 			{
 				struct Local
 				{
@@ -357,7 +357,7 @@ FReply SDocumentationToolTip::ReloadDocumentation()
 void SDocumentationToolTip::Tick( const FGeometry& AllottedGeometry, const double InCurrentTime, const float InDeltaTime )
 {
 	const FModifierKeysState ModifierKeys = FSlateApplication::Get().GetModifierKeys();
-	const bool NeedsUpdate = IsDisplayingDocumentationLink != GEditor->EditorUserSettings->bDisplayDocumentationLink;
+	const bool NeedsUpdate = IsDisplayingDocumentationLink != GetDefault<UEditorPerProjectUserSettings>()->bDisplayDocumentationLink;
 
 	if ( !IsShowingFullTip && ModifierKeys.IsAltDown() && ModifierKeys.IsControlDown() )
 	{
@@ -365,7 +365,7 @@ void SDocumentationToolTip::Tick( const FGeometry& AllottedGeometry, const doubl
 		{
 			ConstructFullTipContent();
 		}
-		else if ( GEditor->EditorUserSettings->bDisplayDocumentationLink )
+		else if ( GetDefault<UEditorPerProjectUserSettings>()->bDisplayDocumentationLink )
 		{
 			ReloadDocumentation();
 		}
@@ -391,7 +391,7 @@ void SDocumentationToolTip::Tick( const FGeometry& AllottedGeometry, const doubl
 		if ( NeedsUpdate )
 		{
 			ReloadDocumentation();
-			IsDisplayingDocumentationLink = GEditor->EditorUserSettings->bDisplayDocumentationLink;
+			IsDisplayingDocumentationLink = GetDefault<UEditorPerProjectUserSettings>()->bDisplayDocumentationLink;
 		}
 
 		WidgetContent->SetContent( SimpleTipContent.ToSharedRef() );
diff --git a/Engine/Source/Editor/Documentation/Private/UDNParser.cpp b/Engine/Source/Editor/Documentation/Private/UDNParser.cpp
--- a/Engine/Source/Editor/Documentation/Private/UDNParser.cpp
+++ b/Engine/Source/Editor/Documentation/Private/UDNParser.cpp
@@ -296,7 +296,7 @@ bool FUDNParser::LoadLink( const FString& Link, TArray<FString>& ContentLines )
 		UDNParserLog.Error(FText::Format(LOCTEXT("LoadingError", "Loading document '{0}' failed."), FText::FromString(SourcePath)));
 	}
 
-	if ( !bLoadSuccess && GEditor->EditorUserSettings->bDisplayDocumentationLink )
+	if ( !bLoadSuccess && GetDefault<UEditorPerProjectUserSettings>()->bDisplayDocumentationLink )
 	{
 		UDNParserLog.Open();
 	}
@@ -322,7 +322,7 @@ bool FUDNParser::Parse(const FString& Link, TArray<FExcerpt>& OutExcerpts, FUDNP
 		}
 		else
 		{
-			if ( GEditor->EditorUserSettings->bDisplayDocumentationLink )
+			if ( GetDefault<UEditorPerProjectUserSettings>()->bDisplayDocumentationLink )
 			{
 				UDNParserLog.Open();
 			}
@@ -347,7 +347,7 @@ bool FUDNParser::GetExcerptContent( const FString& Link, FExcerpt& Excerpt )
 	}
 	else
 	{
-		if ( GEditor->EditorUserSettings->bDisplayDocumentationLink )
+		if ( GetDefault<UEditorPerProjectUserSettings>()->bDisplayDocumentationLink )
 		{
 			UDNParserLog.Open();
 		}
@@ -529,7 +529,7 @@ bool FUDNParser::ParseLineIntoSymbols(int32 LineNumber, const FString& Line, TAr
 			FMessageLog UDNParserLog(UDNParseErrorLog);
 			UDNParserLog.Error(FText::Format(LOCTEXT("TokenParseError", "Line {0}: Token '{1}' could not be parsed properly."), FText::AsNumber(LineNumber), FText::FromString(Line)));
 
-			if ( GEditor->EditorUserSettings->bDisplayDocumentationLink )
+			if ( GetDefault<UEditorPerProjectUserSettings>()->bDisplayDocumentationLink )
 			{
 				UDNParserLog.Open();
 			}
@@ -648,7 +648,7 @@ FUDNLine FUDNParser::ParseLineIntoUDNContent(int32 LineNumber, const FString& Li
 					}
 					else
 					{
-						if ( GEditor->EditorUserSettings->bDisplayDocumentationLink )
+						if ( GetDefault<UEditorPerProjectUserSettings>()->bDisplayDocumentationLink )
 						{
 							UDNParserLog.Open();
 						}
@@ -674,7 +674,7 @@ FUDNLine FUDNParser::ParseLineIntoUDNContent(int32 LineNumber, const FString& Li
 	}
 	else
 	{
-		if ( GEditor->EditorUserSettings->bDisplayDocumentationLink )
+		if ( GetDefault<UEditorPerProjectUserSettings>()->bDisplayDocumentationLink )
 		{
 			UDNParserLog.Open();
 		}
@@ -1021,7 +1021,7 @@ TSharedRef< SWidget > FUDNParser::GenerateExcerptContent( const FString& InLink,
 		bCriticalError = true;
 	}
 
-	if ( bCriticalError && GEditor->EditorUserSettings->bDisplayDocumentationLink )
+	if ( bCriticalError && GetDefault<UEditorPerProjectUserSettings>()->bDisplayDocumentationLink )
 	{
 		UDNParserLog.Open();
 	}
diff --git a/Engine/Source/Editor/EditorLiveStreaming/EditorLiveStreamingSettings.cpp b/Engine/Source/Editor/EditorLiveStreaming/EditorLiveStreamingSettings.cpp
--- a/Engine/Source/Editor/EditorLiveStreaming/EditorLiveStreamingSettings.cpp
+++ b/Engine/Source/Editor/EditorLiveStreaming/EditorLiveStreamingSettings.cpp
@@ -19,9 +19,5 @@ void UEditorLiveStreamingSettings::PostEditChangeProperty( struct FPropertyChang
 //		...
 //	}
 
-	// Save config to file, but only if we are not the build machine since game agnostic settings may put the builder in an unclean state
-	if( !GIsBuildMachine )
-	{
-		this->SaveConfig();
-	}	
+	this->SaveConfig();
 }
\ No newline at end of file
diff --git a/Engine/Source/Editor/EditorLiveStreaming/EditorLiveStreamingSettings.h b/Engine/Source/Editor/EditorLiveStreaming/EditorLiveStreamingSettings.h
--- a/Engine/Source/Editor/EditorLiveStreaming/EditorLiveStreamingSettings.h
+++ b/Engine/Source/Editor/EditorLiveStreaming/EditorLiveStreamingSettings.h
@@ -35,7 +35,7 @@ namespace EEditorLiveStreamingWebCamResolution
 
 
 /** Holds preferences for the editor live streaming features */
-UCLASS(config=EditorGameAgnostic)
+UCLASS(config=EditorSettings)
 class UEditorLiveStreamingSettings : public UObject
 {
 	GENERATED_UCLASS_BODY()
diff --git a/Engine/Source/Editor/EditorSettingsViewer/Private/EditorSettingsViewerModule.cpp b/Engine/Source/Editor/EditorSettingsViewer/Private/EditorSettingsViewerModule.cpp
--- a/Engine/Source/Editor/EditorSettingsViewer/Private/EditorSettingsViewerModule.cpp
+++ b/Engine/Source/Editor/EditorSettingsViewer/Private/EditorSettingsViewerModule.cpp
@@ -142,7 +142,7 @@ protected:
 		SettingsModule.RegisterSettings("Editor", "General", "UserSettings",
 			LOCTEXT("UserSettingsName", "Miscellaneous"),
 			LOCTEXT("UserSettingsDescription", "Customize the behavior, look and feel of the editor."),
-			&GEditor->AccessEditorUserSettings()
+			GetMutableDefault<UEditorPerProjectUserSettings>()
 		);
 
 		// experimental features
