Correcting section path for editor project appearance setings

diff --git a/Engine/Config/BaseEditor.ini b/Engine/Config/BaseEditor.ini
--- a/Engine/Config/BaseEditor.ini
+++ b/Engine/Config/BaseEditor.ini
@@ -420,7 +420,7 @@ MaxMemoryAllowance=8192
 +EngineTargetsSettings=(Name="Keywords",Guid=AE89AECB47475F420D0D69A5547515DC,TargetDependencies=(33482D004789784C9DA695A682ACCA1B,AC8BFD2A41A2FB2893BB8EA0AF903E6D),AdditionalManifestDependencies=,RequiredModuleNames=,GatherFromTextFiles=(IsEnabled=False,SearchDirectories=,ExcludePathWildcards=,FileExtensions=((Pattern="h"),(Pattern="cpp"),(Pattern="ini"))),GatherFromPackages=(IsEnabled=False,IncludePathWildcards=,ExcludePathWildcards=,FileExtensions=((Pattern="umap"),(Pattern="uasset"))),GatherFromMetaData=(IsEnabled=True,IncludePathWildcards=((Pattern="Source/Editor/*"),(Pattern="Source/Runtime/*"),(Pattern="Source/Developer/*")),ExcludePathWildcards=((Pattern="Source/Developer/NoRedist/CommunityPortalServices/*")),KeySpecifications=((MetaDataKey=(Name="Keywords"),TextNamespace="UObjectKeywords",TextKeyPattern=(Pattern="{FieldPath}")))),NativeCultureIndex=0,SupportedCulturesStatistics=((CultureName="en"),(CultureName="es"),(CultureName="ja"),(CultureName="ko"),(CultureName="pt-BR"),(CultureName="zh-CN")))
 +GameTargetsSettings=(Name="Game",Guid=AE0EA34A45461A25BA65A391026F19F8,TargetDependencies=(33482D004789784C9DA695A682ACCA1B,AC8BFD2A41A2FB2893BB8EA0AF903E6D),AdditionalManifestDependencies=,RequiredModuleNames=,GatherFromTextFiles=(IsEnabled=False,SearchDirectories=,ExcludePathWildcards=,FileExtensions=((Pattern="h"),(Pattern="cpp"),(Pattern="ini"))),GatherFromPackages=(IsEnabled=False,IncludePathWildcards=,ExcludePathWildcards=,FileExtensions=((Pattern="umap"),(Pattern="uasset"))),GatherFromMetaData=(IsEnabled=False,IncludePathWildcards=,ExcludePathWildcards=,KeySpecifications=),NativeCultureIndex=,SupportedCulturesStatistics=)
 
-[/Script/EditorProjectAppearanceSettings.EditorProjectAppearanceSettings]
+[/Script/UnrealEd.EditorProjectAppearanceSettings]
 bDisplayUnits=True
 +DistanceUnits=EUnit::Centimeters
 +MassUnits=EUnit::Kilograms
