Renamed/fixed up remaining EditorUserSettings references in the codebase

EditorUserSettings is now EditorPerProjectUserSettings

diff --git a/Engine/Config/BaseEditor.ini b/Engine/Config/BaseEditor.ini
--- a/Engine/Config/BaseEditor.ini
+++ b/Engine/Config/BaseEditor.ini
@@ -357,7 +357,7 @@ MaxMemoryAllowance=8192
 +VersionedIniParams=Engine:/Script/AndroidRuntimeSettings.AndroidRuntimeSettings:AndroidAudio
 +VersionedIniParams=Engine:/Script/UnrealEd.CookerSettings:DefaultPVRTCQuality
 +VersionedIniParams=Editor:StructSerialization:BreakSerializationRecursion
-+VersionedIniParams=EditorUserSettings:/Script/UnrealEd.EditorUserSettings:MaterialQualityLevel
++VersionedIniParams=EditorPerProjectUserSettings:/Script/UnrealEd.EditorPerProjectUserSettings:MaterialQualityLevel
 +VersionedIniParams=Engine:/Script/Engine.Engine:GameUserSettingsClassName
 +VersionedIniParams=Engine:/Script/EngineSettings.GameMapsSettings:Compat.UseDXT5NormalMaps
 +VersionedIniParams=Engine:/Script/EngineSettings.GameMapsSettings:GameDefaultMap
diff --git a/Engine/Config/BaseEditorPerProjectUserSettings.ini b/Engine/Config/BaseEditorPerProjectUserSettings.ini
--- a/Engine/Config/BaseEditorPerProjectUserSettings.ini
+++ b/Engine/Config/BaseEditorPerProjectUserSettings.ini
@@ -1,16 +1,16 @@
 
 ; -------------------------------------------------------------------------------------------------
-; IMPORTANT:  EditorUserSettings.ini has special behavior!
+; IMPORTANT:  EditorPerProjectUserSettings.ini has special behavior!
 ;
 ; This .ini file contains the defaults for many editor and engine "preferences".  These preferences
-; are saved to an EditorUserSettings.ini file in the user's /<Game>/Saved/Config/ directory. 
+; are saved to an EditorPerProjectUserSettings.ini file in the user's /<Game>/Saved/Config/ directory. 
 ;
 ; If you change the default for a preference in here, that change will *NOT* automatically 
 ; propagate to users who already have saved preferences.  The idea is that we must preserve
 ; their existing settings and never want to clobber entries in that file.
 ;
 ; If you add a new preference, the new key and value *WILL* be propagated to the user's
-; EditorUserSettings.ini file.  After we load the user's existing settings, we'll merge in
+; EditorPerProjectUserSettings.ini file.  After we load the user's existing settings, we'll merge in
 ; any missing keys and values that are present in these defaults.
 ;
 ; One easy technique for forcing a "reset" of an already-existing preference is to simply
@@ -18,7 +18,7 @@
 ; key and will propagate the change to existing users' preferences.
 ; -------------------------------------------------------------------------------------------------
 
-[/Script/UnrealEd.EditorUserSettings]
+[/Script/UnrealEd.EditorPerProjectUserSettings]
 ; True if WASD keys should be remapped to flight controls while the right mouse button is held down
 AllowFlightCameraToRemapKeys=True
 
diff --git a/Engine/Config/BaseEngine.ini b/Engine/Config/BaseEngine.ini
--- a/Engine/Config/BaseEngine.ini
+++ b/Engine/Config/BaseEngine.ini
@@ -2000,7 +2000,7 @@ Enabled=false
 ConfigPathPrefix=\\epicgames.net\root\Home
 ConfigPathSuffix=UE4Cloud
 Timeout=1.0f
-+IniToLoad=EditorUserSettings
++IniToLoad=EditorPerProjectUserSettings
 +IniToLoad=EditorKeyBindings
 
 [Engine.ErrorHandling]
diff --git a/Engine/Config/Mac/MacEditorPerProjectUserSettings.ini b/Engine/Config/Mac/MacEditorPerProjectUserSettings.ini
new file mode 100644
--- /dev/null
+++ b/Engine/Config/Mac/MacEditorPerProjectUserSettings.ini
@@ -0,0 +1,23 @@
+
+; -------------------------------------------------------------------------------------------------
+; IMPORTANT:  EditorPerProjectUserSettings.ini has special behavior!
+;
+; This .ini file contains the defaults for many editor and engine "preferences".  These preferences
+; are saved to an EditorPerProjectUserSettings.ini file in the user's /<Game>/Saved/Config/ directory. 
+;
+; If you change the default for a preference in here, that change will *NOT* automatically 
+; propagate to users who already have saved preferences.  The idea is that we must preserve
+; their existing settings and never want to clobber entries in that file.
+;
+; If you add a new preference, the new key and value *WILL* be propagated to the user's
+; EditorPerProjectUserSettings.ini file.  After we load the user's existing settings, we'll merge in
+; any missing keys and values that are present in these defaults.
+;
+; One easy technique for forcing a "reset" of an already-existing preference is to simply
+; rename the variable for that preference.  The config system will treat it as a newly-added
+; key and will propagate the change to existing users' preferences.
+; -------------------------------------------------------------------------------------------------
+
+[/Script/UnrealEd.EditorLoadingSavingSettings]
+; Tool to use for diff'ing text
+TextDiffToolPath=(FilePath="p4merge.app")
\ No newline at end of file
diff --git a/Engine/Config/Mac/MacEditorUserSettings.ini b/Engine/Config/Mac/MacEditorUserSettings.ini
deleted file mode 100644
