Fixed missing ) in baseeditor.ini

diff --git a/Engine/Config/BaseEditor.ini b/Engine/Config/BaseEditor.ini
--- a/Engine/Config/BaseEditor.ini
+++ b/Engine/Config/BaseEditor.ini
@@ -45,6 +45,7 @@ DefaultSizeX=512
 DefaultSizeY=512
 
 [/Script/UnrealEd.CollisionAutomationTestConfigData]
+
 // Component sweeps
 // Capsule
 +ComponentSweepMultiTests=(RootShapeAsset="/Engine/BasicShapes/Cone.Cone",ShapeType="TriggerCapsule",HitResult=(Time=0.25,Location=(X=0.0,Y=0.0,Z=149.999939),ImpactPoint=(X=0.070709,Y=0.0,Z=49.999985),Normal=(X=-0.001414,Y=-0.0,Z=0.999999),ImpactNormal=(X=0.0,Y=0.0,Z=1.0),TraceStart=(X=0.0,Y=0.0,Z=200.0),TraceEnd=(X=0.0,Y=0.0,Z=0.0)))
@@ -53,7 +54,7 @@ DefaultSizeY=512
 
 +ComponentSweepMultiTests=(RootShapeAsset="/Engine/BasicShapes/Cube.Cube",ShapeType="TriggerCapsule",HitResult=(Time=0.25,Location=(X=0.0,Y=0.0,Z=150.0),ImpactPoint=(X=0.0,Y=0.0,Z=50.0),Normal=(X=0.0,Y=-0.0,Z=1.0),ImpactNormal=(X=0.0,Y=0.0,Z=1.0),TraceStart=(X=0.0,Y=0.0,Z=200.0),TraceEnd=(X=0.0,Y=0.0,Z=0.0)))
 +ComponentSweepMultiTests=(RootShapeAsset="/Engine/BasicShapes/Cube.Cube",ShapeType="TriggerCapsule",HitResult=(Time=0.50,Location=(X=-100.000008,Y=0.0,Z=0.0),ImpactPoint=(X=-50.0,Y=0.0,Z=0.0),Normal=(X=-1.0,Y=0.0,Z=0.0),ImpactNormal=(X=-1.0,Y=0.0,Z=0.0),TraceStart=(X=-200.0,Y=0.0,Z=0.0),TraceEnd=(X=0.0,Y=0.0,Z=0.0),PenetrationDepth=0.0)))
-+ComponentSweepMultiTests=(RootShapeAsset="/Engine/BasicShapes/Cube.Cube",ShapeType="TriggerCapsule",HitResult=(Time=0.50,Location=(X=-100.000015,Y=0.0,Z=100.000015),ImpactPoint=(X=-50.0,Y=0.0,Z=50.0),Normal=(X=-1.0,Y=-0.0,Z=0.0),ImpactNormal=(X=0.0,Y=0.0,Z=1.0),TraceStart=(X=-200.0,Y=0.0,Z=200.0),TraceEnd=(X=0.0,Y=0.0,Z=0.0))
++ComponentSweepMultiTests=(RootShapeAsset="/Engine/BasicShapes/Cube.Cube",ShapeType="TriggerCapsule",HitResult=(Time=0.50,Location=(X=-100.000015,Y=0.0,Z=100.000015),ImpactPoint=(X=-50.0,Y=0.0,Z=50.0),Normal=(X=-1.0,Y=-0.0,Z=0.0),ImpactNormal=(X=0.0,Y=0.0,Z=1.0),TraceStart=(X=-200.0,Y=0.0,Z=200.0),TraceEnd=(X=0.0,Y=0.0,Z=0.0)))
 
 +ComponentSweepMultiTests=(RootShapeAsset="/Engine/BasicShapes/Cylinder.Cylinder",ShapeType="TriggerCapsule",HitResult=(Time=0.25,Location=(X=0.0,Y=0.0,Z=150.0),ImpactPoint=(X=0.0,Y=0.0,Z=50.0),Normal=(X=0.0,Y=-0.0,Z=1.0),ImpactNormal=(X=0.0,Y=0.0,Z=1.0),TraceStart=(X=0.0,Y=0.0,Z=200.0),TraceEnd=(X=0.0,Y=0.0,Z=0.0)))
 +ComponentSweepMultiTests=(RootShapeAsset="/Engine/BasicShapes/Cylinder.Cylinder",ShapeType="TriggerCapsule",HitResult=(Time=0.50,Location=(X=-100.0,Y=0.0,Z=0.0),ImpactPoint=(X=-50.0,Y=0.000016,Z=30.0),Normal=(X=-1.0,Y=-0.0,Z=0.0),ImpactNormal=(X=-0.995185,Y=0.098017,Z=0.0),TraceStart=(X=-200.0,Y=0.0,Z=0.0),TraceEnd=(X=0.0,Y=0.0,Z=0.0)))
@@ -87,7 +88,7 @@ DefaultSizeY=512
 
 +LineTraceSingleByChannelTests=(RootShapeAsset="/Engine/BasicShapes/Cube.Cube",HitResult=(Time=0.25,Location=(X=0.0,Y=0.0,Z=150.0),ImpactPoint=(X=0.0,Y=0.0,Z=50.0),Normal=(X=0.0,Y=-0.0,Z=1.0),ImpactNormal=(X=0.0,Y=0.0,Z=1.0),TraceStart=(X=0.0,Y=0.0,Z=200.0),TraceEnd=(X=0.0,Y=0.0,Z=0.0)))
 +LineTraceSingleByChannelTests=(RootShapeAsset="/Engine/BasicShapes/Cube.Cube",HitResult=(Time=0.50,Location=(X=-100.000008,Y=0.0,Z=0.0),ImpactPoint=(X=-50.0,Y=0.0,Z=0.0),Normal=(X=-1.0,Y=0.0,Z=0.0),ImpactNormal=(X=-1.0,Y=0.0,Z=0.0),TraceStart=(X=-200.0,Y=0.0,Z=0.0),TraceEnd=(X=0.0,Y=0.0,Z=0.0),PenetrationDepth=0.0)))
-+LineTraceSingleByChannelTests=(RootShapeAsset="/Engine/BasicShapes/Cube.Cube",HitResult=(Time=0.50,Location=(X=-100.000015,Y=0.0,Z=100.000015),ImpactPoint=(X=-50.0,Y=0.0,Z=50.0),Normal=(X=-1.0,Y=-0.0,Z=0.0),ImpactNormal=(X=0.0,Y=0.0,Z=1.0),TraceStart=(X=-200.0,Y=0.0,Z=200.0),TraceEnd=(X=0.0,Y=0.0,Z=0.0))
++LineTraceSingleByChannelTests=(RootShapeAsset="/Engine/BasicShapes/Cube.Cube",HitResult=(Time=0.50,Location=(X=-100.000015,Y=0.0,Z=100.000015),ImpactPoint=(X=-50.0,Y=0.0,Z=50.0),Normal=(X=-1.0,Y=-0.0,Z=0.0),ImpactNormal=(X=0.0,Y=0.0,Z=1.0),TraceStart=(X=-200.0,Y=0.0,Z=200.0),TraceEnd=(X=0.0,Y=0.0,Z=0.0)))
 
 +LineTraceSingleByChannelTests=(RootShapeAsset="/Engine/BasicShapes/Cylinder.Cylinder",HitResult=(Time=0.25,Location=(X=0.0,Y=0.0,Z=150.0),ImpactPoint=(X=0.0,Y=0.0,Z=50.0),Normal=(X=0.0,Y=-0.0,Z=1.0),ImpactNormal=(X=0.0,Y=0.0,Z=1.0),TraceStart=(X=0.0,Y=0.0,Z=200.0),TraceEnd=(X=0.0,Y=0.0,Z=0.0)))
 +LineTraceSingleByChannelTests=(RootShapeAsset="/Engine/BasicShapes/Cylinder.Cylinder",HitResult=(Time=0.50,Location=(X=-100.0,Y=0.0,Z=0.0),ImpactPoint=(X=-50.0,Y=0.000016,Z=30.0),Normal=(X=-1.0,Y=-0.0,Z=0.0),ImpactNormal=(X=-0.995185,Y=0.098017,Z=0.0),TraceStart=(X=-200.0,Y=0.0,Z=0.0),TraceEnd=(X=0.0,Y=0.0,Z=0.0)))
