UMG - UE-14779 - Adding iPhone 6/6+ to the hardware resolution list, also adds it for previewing game.

diff --git a/Engine/Config/BaseEditorPerProjectUserSettings.ini b/Engine/Config/BaseEditorPerProjectUserSettings.ini
--- a/Engine/Config/BaseEditorPerProjectUserSettings.ini
+++ b/Engine/Config/BaseEditorPerProjectUserSettings.ini
@@ -222,6 +222,10 @@ LastExecutedLaunchPlatform=
 +PhoneScreenResolutions=(Description="Apple iPhone 4, 4S, iTouch 4 (Landscape)",Width=960,Height=640,AspectRatio="3:2")
 +PhoneScreenResolutions=(Description="Apple iPhone 5, 5S, iTouch 5 (Portrait)",Width=640,Height=1136,AspectRatio="~9:16")
 +PhoneScreenResolutions=(Description="Apple iPhone 5, 5S, iTouch 5 (Landscape)",Width=1136,Height=640,AspectRatio="~16:9")
++PhoneScreenResolutions=(Description="Apple iPhone 6 (Portrait)",Width=750,Height=1334,AspectRatio="16:9")
++PhoneScreenResolutions=(Description="Apple iPhone 6 (Landscape)",Width=1334,Height=750,AspectRatio="16:9")
++PhoneScreenResolutions=(Description="Apple iPhone 6+ (Portrait)",Width=1080,Height=1920,AspectRatio="16:9")
++PhoneScreenResolutions=(Description="Apple iPhone 6+ (Landscape)",Width=1920,Height=1080,AspectRatio="16:9")
 +PhoneScreenResolutions=(Description="HTC One (Portrait)",Width=1080,Height=1920,AspectRatio="9:16")
 +PhoneScreenResolutions=(Description="HTC One (Landscape)",Width=1920,Height=1080,AspectRatio="16:9")
 +PhoneScreenResolutions=(Description="Samsung Galaxy S4 (Portrait)",Width=1080,Height=1920,AspectRatio="9:16")
