- Adding some text to Base.ini so that it doesn't cause problems by being empty and then getting skipped, etc. [user reported]

diff --git a/Engine/Config/Base.ini b/Engine/Config/Base.ini
--- a/Engine/Config/Base.ini
+++ b/Engine/Config/Base.ini
@@ -0,0 +1 @@
+; Base.ini is required to exist
