Paper2D: Fixed missing collision/render shapes when loading assets saved with Paper2D custom version 13 (4.7.4 and above) [UE-15311]

diff --git a/Engine/Config/BaseEngine.ini b/Engine/Config/BaseEngine.ini
--- a/Engine/Config/BaseEngine.ini
+++ b/Engine/Config/BaseEngine.ini
@@ -613,8 +613,8 @@ bUseStreamingPause=false
 +ActiveClassRedirects=(OldClassName="PaperTileMapRenderComponent",NewClassName="/Script/Paper2D.PaperTileMapComponent")
 
 ; Paper2D modifications (4.8 era)
-+ActiveClassRedirects=(OldClassName="SpritePolygon",NewClassName="/Script/Paper2D.SpriteGeometryShape")
-+ActiveClassRedirects=(OldClassName="SpritePolygonCollection",NewClassName="/Script/Paper2D.SpriteGeometryCollection")
++ActiveStructRedirects=(OldStructName="SpritePolygon",NewStructName="SpriteGeometryShape")
++ActiveStructRedirects=(OldStructName="SpritePolygonCollection",NewStructName="SpriteGeometryCollection")
 +TaggedPropertyRedirects=(ClassName="SpriteGeometryCollection",OldPropertyName="Polygons",NewPropertyName="Shapes")
 +TaggedPropertyRedirects=(ClassName="PaperTileSet",OldPropertyName="ExperimentalPerTileData",NewPropertyName="PerTileData")
 
