Enable vtable ptr ctors by default.

diff --git a/Engine/Config/BaseEngine.ini b/Engine/Config/BaseEngine.ini
--- a/Engine/Config/BaseEngine.ini
+++ b/Engine/Config/BaseEngine.ini
@@ -1326,6 +1326,7 @@ FlushStreamingOnGC=False
 MinBulkDataSizeForAsyncLoading=131072
 AsyncLoadingThreadEnabled=True
 WarnIfTimeLimitExceeded=True
+UseVTableConstructors=True
 
 [Internationalization]
 +LocalizationPaths=../../../Engine/Content/Localization/Engine
