UE4 - change default shared DDC lifetime to 10 days

diff --git a/Engine/Config/BaseEngine.ini b/Engine/Config/BaseEngine.ini
--- a/Engine/Config/BaseEngine.ini
+++ b/Engine/Config/BaseEngine.ini
@@ -1849,7 +1849,7 @@ AsyncPut=(Type=AsyncPut, Inner=Hierarchy)
 Hierarchy=(Type=Hierarchical, Inner=Boot, Inner=Pak, Inner=EnginePak, Inner=Local, Inner=Shared)
 Boot=(Type=Boot, Filename="%GAMEDIR%DerivedDataCache/Boot.ddc", MaxCacheSize=512)
 Local=(Type=FileSystem, ReadOnly=false, Clean=false, Flush=false, PurgeTransient=true, DeleteUnused=true, UnusedFileAge=34, FoldersToClean=-1, Path=../../../Engine/DerivedDataCache, EnvPathOverride=UE-LocalDataCachePath)
-Shared=(Type=FileSystem, ReadOnly=false, Clean=false, Flush=false, DeleteUnused=true, UnusedFileAge=23, FoldersToClean=10, MaxFileChecksPerSec=1, Path=?EpicDDC, EnvPathOverride=UE-SharedDataCachePath)
+Shared=(Type=FileSystem, ReadOnly=false, Clean=false, Flush=false, DeleteUnused=true, UnusedFileAge=10, FoldersToClean=10, MaxFileChecksPerSec=1, Path=?EpicDDC, EnvPathOverride=UE-SharedDataCachePath)
 AltShared=(Type=FileSystem, ReadOnly=true, Clean=false, Flush=false, DeleteUnused=true, UnusedFileAge=23, FoldersToClean=10, MaxFileChecksPerSec=1, Path=?EpicDDC2, EnvPathOverride=UE-SharedDataCachePath2)
 Pak=(Type=ReadPak, Filename="%GAMEDIR%DerivedDataCache/DDC.ddp")
 EnginePak=(Type=ReadPak, Filename=../../../Engine/DerivedDataCache/DDC.ddp)
