Async Loading Thread disabled by default

diff --git a/Engine/Config/BaseEngine.ini b/Engine/Config/BaseEngine.ini
--- a/Engine/Config/BaseEngine.ini
+++ b/Engine/Config/BaseEngine.ini
@@ -1326,7 +1326,7 @@ ZeroEngineVersionWarning=True
 UseStrictEngineVersioning=True
 FlushStreamingOnGC=False
 MinBulkDataSizeForAsyncLoading=131072
-AsyncLoadingThreadEnabled=True
+AsyncLoadingThreadEnabled=False
 WarnIfTimeLimitExceeded=True
 UseVTableConstructors=True
 
