Added support for some Rendering RValues to invalidate cooked builds.

diff --git a/Engine/Config/BaseEditor.ini b/Engine/Config/BaseEditor.ini
--- a/Engine/Config/BaseEditor.ini
+++ b/Engine/Config/BaseEditor.ini
@@ -419,6 +419,19 @@ bUseFullGCAssetClassNames=true
 +VersionedIniParams=Engine:/Script/Engine.Engine:WeightMapPlaceholderTextureName
 +VersionedIniParams=Engine:/Script/Engine.Engine:LightMapDensityTextureName
 
+
++VersionedIntRValues=r.AllowStaticLighting
++VersionedIntRValues=r.GBuffer
++VersionedIntRValues=r.BasePassOutputsVelocity
++VersionedIntRValues=r.DBuffer
++VersionedIntRValues=r.Shaders.KeepDebugInfo
++VersionedIntRValues=r.Shaders.Optimize
++VersionedIntRValues=r.CompileShadersForDevelopment
++VersionedIntRValues=r.PS4MixedModeShaderDebugInfo
++VersionedIntRValues=r.PS4DumpShaderSDB
+
+
+
 [/Script/LocalizationDashboard.LocalizationDashboardSettings]
 +EngineTargetsSettings=(Name="Engine",Guid=33482D004789784C9DA695A682ACCA1B,TargetDependencies=,AdditionalManifestDependencies=,RequiredModuleNames=,GatherFromTextFiles=(IsEnabled=True,SearchDirectories=((Path="Source/Runtime/"),(Path="Source/Developer/"),(Path="Config/")),ExcludePathWildcards=((Pattern="Source/Developer/NoRedist/CommunityPortalServices/*"),(Pattern="Source/Developer/NoRedist/UnrealEngineLauncherServices/*"),(Pattern="Source/Developer/NoRedist/BuildPatchServices/*")),FileExtensions=((Pattern="cpp"),(Pattern="h"),(Pattern="c"),(Pattern="inl"),(Pattern="mm"),(Pattern="ini"))),GatherFromPackages=(IsEnabled=True,IncludePathWildcards=((Pattern="Content/*")),ExcludePathWildcards=((Pattern="Content/Editor/*"),(Pattern="Content/Tutorial/*"),(Pattern="Content/Developers/*"),(Pattern="Content/TestPackages/*"),(Pattern="Content/QA_Assets/*"),(Pattern="Content/Maps/Automation/*"),(Pattern="Content/EngineSounds/*")),FileExtensions=((Pattern="umap"),(Pattern="uasset"))),GatherFromMetaData=(IsEnabled=False,IncludePathWildcards=,ExcludePathWildcards=,KeySpecifications=),NativeCultureIndex=0,SupportedCulturesStatistics=((CultureName="en"),(CultureName="es"),(CultureName="ja"),(CultureName="ko"),(CultureName="pt-BR"),(CultureName="zh-CN")))
 +EngineTargetsSettings=(Name="Editor",Guid=AC8BFD2A41A2FB2893BB8EA0AF903E6D,TargetDependencies=(33482D004789784C9DA695A682ACCA1B),AdditionalManifestDependencies=,RequiredModuleNames=,GatherFromTextFiles=(IsEnabled=True,SearchDirectories=((Path="Source/Editor/")),ExcludePathWildcards=,FileExtensions=((Pattern="cpp"),(Pattern="h"),(Pattern="c"),(Pattern="inl"),(Pattern="mm"))),GatherFromPackages=(IsEnabled=True,IncludePathWildcards=((Pattern="Content/Editor/*"),(Pattern="Content/Editor*")),ExcludePathWildcards=,FileExtensions=((Pattern="umap"),(Pattern="uasset"))),GatherFromMetaData=(IsEnabled=True,IncludePathWildcards=((Pattern="Source/Editor/*"),(Pattern="Source/Runtime/*"),(Pattern="Source/Developer/*")),ExcludePathWildcards=((Pattern="Source/Developer/NoRedist/CommunityPortalServices/*"),(Pattern="Source/Developer/NoRedist/UnrealEngineLauncherServices/*"),(Pattern="Source/Developer/NoRedist/BuildPatchServices/*")),KeySpecifications=((MetaDataKey=(Name="Category"),TextNamespace="UObjectCategories",TextKeyPattern=(Pattern="{MetaDataValue}")))),NativeCultureIndex=0,SupportedCulturesStatistics=((CultureName="en"),(CultureName="es"),(CultureName="ja"),(CultureName="ko"),(CultureName="pt-BR"),(CultureName="zh-CN")))
diff --git a/Engine/Source/Editor/UnrealEd/Private/CookOnTheFlyServer.cpp b/Engine/Source/Editor/UnrealEd/Private/CookOnTheFlyServer.cpp
--- a/Engine/Source/Editor/UnrealEd/Private/CookOnTheFlyServer.cpp
+++ b/Engine/Source/Editor/UnrealEd/Private/CookOnTheFlyServer.cpp
@@ -2353,6 +2353,21 @@ bool UCookOnTheFlyServer::GetCurrentIniVersionStrings( const ITargetPlatform* Ta
 		IniVersionStrings.Emplace( MoveTemp(CurrentVersionString) );
 	}
 
+
+	TArray<FString> VersionedRValues;
+	GConfig->GetArray(TEXT("CookSettings"), TEXT("VersionedIntRValues"), VersionedRValues, GEditorIni);
+
+	for (const auto& RValue : VersionedRValues)
+	{
+		static const auto CVar = IConsoleManager::Get().FindTConsoleVariableDataInt(TEXT("r.CompileShadersForDevelopment"));
+		if (CVar)
+		{
+			FString VersionedRValueString = FString::Printf(TEXT("%s:%d"), *RValue, CVar->GetValueOnGameThread());
+			IniVersionStrings.Emplace(MoveTemp(VersionedRValueString));
+		}
+	}
+
+
 	// clean up our temporary platform ini files
 	for ( const auto& PlatformIniFile : PlatformIniFiles )
 	{
@@ -2555,11 +2570,11 @@ void UCookOnTheFlyServer::PopulateCookedPackagesFromDisk( const TArray<ITargetPl
 			CookedFilename = FPaths::ConvertRelativePathToFull(CookedFilename);
 			FString StandardCookedFilename = CookedFilename;
 
-			StandardCookedFilename.ReplaceInline(*EngineSandboxPath, *LocalEnginePath);
-			StandardCookedFilename.ReplaceInline(*GameSandboxPath, *LocalGamePath);
 			StandardCookedFilename.ReplaceInline(*SandboxPath, *(FPaths::GetRelativePathToRoot()));
-			FDateTime DependentTimestamp;
 
+			StandardCookedFilename = FPaths::ConvertRelativePathToFull(StandardCookedFilename);
+			FPaths::MakeStandardFilename(StandardCookedFilename);
+			FDateTime DependentTimestamp;
 			if (PDInfoModule.DeterminePackageDependentTimeStamp(*(FPaths::GetBaseFilename(StandardCookedFilename, false)), DependentTimestamp) == true)
 			{
 				double Diff = (CookedTimestamp - DependentTimestamp).GetTotalSeconds();
@@ -2674,6 +2689,7 @@ void UCookOnTheFlyServer::CleanSandbox( const bool bIterative )
 						(CookedFilename.EndsWith(TEXT(".umap")) == false))
 					{
 						// don't care if it's not a map or an asset
+						//IFileManager::Get().Delete(*CookedFilename);
 						continue;
 					}
 
@@ -2681,10 +2697,11 @@ void UCookOnTheFlyServer::CleanSandbox( const bool bIterative )
 					CookedFilename = FPaths::ConvertRelativePathToFull(CookedFilename);
 					FString StandardCookedFilename = CookedFilename;
 
-					StandardCookedFilename.ReplaceInline(*EngineSandboxPath, *LocalEnginePath);
-					StandardCookedFilename.ReplaceInline(*GameSandboxPath, *LocalGamePath);
 					StandardCookedFilename.ReplaceInline(*SandboxPath, *(FPaths::GetRelativePathToRoot()));
 
+					StandardCookedFilename = FPaths::ConvertRelativePathToFull(StandardCookedFilename);
+					FPaths::MakeStandardFilename(StandardCookedFilename);
+					
 					FDateTime DependentTimestamp;
 					FName StandardCookedFileFName = FName(*StandardCookedFilename);
 
