Adding default settings for Material and Particle promotion tests.
Adding Cascade to EditorTestModules, as its automated tests will not appear until the module is loaded.
New test settings use Engine content to allow the tests to be run on any project.

diff --git a/Engine/Config/BaseEngine.ini b/Engine/Config/BaseEngine.ini
--- a/Engine/Config/BaseEngine.ini
+++ b/Engine/Config/BaseEngine.ini
@@ -1962,8 +1962,11 @@ TestAllBlueprints=false
 +EditorTestModules=StaticMeshEditor
 +EditorTestModules=LandscapeEditor
 +EditorTestModules=GameProjectGeneration
++EditorTestModules=Cascade
 +TestLevelFolders=TestMaps
 +EditorPerformanceTestMaps=(PerformanceTestmap=(FilePath="CurrentLevel"),TestTimer=60)
+MaterialEditorPromotionTest=(DefaultMaterialAsset=(FilePath="../../Content/EditorMeshes/ColorCalibrator/M_ColorGrid.uasset"),DefaultDiffuseTexture=(FilePath="../../Content/EngineMaterials/DefaultDiffuse.uasset"),DefaultNormalTexture=(FilePath="../../Content/EngineMaterials/DefaultNormal.uasset"))
+ParticleEditorPromotionTest=(DefaultParticleAsset=(FilePath="../../Content/Tutorial/SubEditors/TutorialAssets/TutorialParticleSystem.uasset"))
 
 [AutomationTesting.StaticMeshEditorTest]
 +EditorViewButtons=Wireframe
