Swapped 45 and 90 rotation snaps to common in .ini file.

UE-1007

diff --git a/Engine/Config/BaseEditorPerProjectUserSettings.ini b/Engine/Config/BaseEditorPerProjectUserSettings.ini
--- a/Engine/Config/BaseEditorPerProjectUserSettings.ini
+++ b/Engine/Config/BaseEditorPerProjectUserSettings.ini
@@ -313,14 +313,15 @@ bUsePowerOf2SnapSize=False
 .CommonRotGridSizes=10
 .CommonRotGridSizes=15
 .CommonRotGridSizes=30
+.CommonRotGridSizes=45
 .CommonRotGridSizes=60
+.CommonRotGridSizes=90
 .CommonRotGridSizes=120
 .DivisionsOf360RotGridSizes=2.8125
 .DivisionsOf360RotGridSizes=5.625
 .DivisionsOf360RotGridSizes=11.25
 .DivisionsOf360RotGridSizes=22.5
-.DivisionsOf360RotGridSizes=45
-.DivisionsOf360RotGridSizes=90
+
 .ScalingGridSizes=10
 .ScalingGridSizes=1
 .ScalingGridSizes=0.5
