Fixed automated LineTest data

diff --git a/Engine/Config/BaseEditor.ini b/Engine/Config/BaseEditor.ini
--- a/Engine/Config/BaseEditor.ini
+++ b/Engine/Config/BaseEditor.ini
@@ -83,21 +83,21 @@ DefaultSizeY=512
 
 
 //Line traces
-+LineTraceSingleByChannelTests=(RootShapeAsset="/Engine/BasicShapes/Cone.Cone",HitResult=(Time=0.25,Location=(X=0.0,Y=0.0,Z=149.999939),ImpactPoint=(X=0.070709,Y=0.0,Z=49.999985),Normal=(X=-0.001414,Y=-0.0,Z=0.999999),ImpactNormal=(X=0.0,Y=0.0,Z=1.0),TraceStart=(X=0.0,Y=0.0,Z=200.0),TraceEnd=(X=0.0,Y=0.0,Z=0.0)))
-+LineTraceSingleByChannelTests=(RootShapeAsset="/Engine/BasicShapes/Cone.Cone",HitResult=(Time=0.50,Location=(X=-100.0,Y=0.0,Z=0.0),ImpactPoint=(X=-50.0,Y=-0.000019,Z=-50.0),Normal=(X=-1.0,Y=0.0,Z=-0.0),ImpactNormal=(X=-0.888290,Y=0.116945,Z=0.444145),TraceStart=(X=-200.0,Y=0.0,Z=0.0),TraceEnd=(X=0.0,Y=0.0,Z=0.0)))
-+LineTraceSingleByChannelTests=(RootShapeAsset="/Engine/BasicShapes/Cone.Cone",HitResult=(Time=0.646994,Location=(X=-70.601135,Y=0.0,Z=70.601135),ImpactPoint=(X=-25.879774,Y=-0.000010,Z=-1.759551),Normal=(X=-0.894427,Y=0.0,Z=0.447214),ImpactNormal=(X=-0.888290,Y=0.116945,Z=0.444145),TraceStart=(X=-200.0,Y=0.0,Z=200.0),TraceEnd=(X=0.0,Y=0.0,Z=0.0)))
++LineTraceSingleByChannelTests=(RootShapeAsset="/Engine/BasicShapes/Cone.Cone",HitResult=(Time=0.750,Location=(X=0.0,Y=0.0,Z=49.999985),ImpactPoint=(X=0.0,Y=-0.0,Z=49.999989),Normal=(X=0.0,Y=0.0,Z=1.000000),ImpactNormal=(X=0.0,Y=0.0,Z=1.000000),TraceStart=(X=0.0,Y=0.0,Z=200.0),TraceEnd=(X=0.0,Y=0.0,Z=0.0),PenetrationDepth=0.0))
++LineTraceSingleByChannelTests=(RootShapeAsset="/Engine/BasicShapes/Cone.Cone",HitResult=(Time=0.875000,Location=(X=-24.999985,Y=0.0,Z=0.0),ImpactPoint=(X=-24.999996,Y=-0.0,Z=0.0),Normal=(X=-0.888290,Y=0.116945,Z=0.444145),ImpactNormal=(X=-0.888290,Y=0.116945,Z=0.444145),TraceStart=(X=-200.0,Y=0.0,Z=0.0),TraceEnd=(X=0.0,Y=0.0,Z=0.0),PenetrationDepth=0.0))
++LineTraceSingleByChannelTests=(RootShapeAsset="/Engine/BasicShapes/Cone.Cone",HitResult=(Time=0.916667,Location=(X=-16.666687,Y=0.0,Z=16.666687),ImpactPoint=(X=-16.666664,Y=-0.000001,Z=16.666666),Normal=(X=-0.888290,Y=0.116945,Z=0.444145),ImpactNormal=(X=-0.888290,Y=0.116945,Z=0.444145),TraceStart=(X=-200.0,Y=0.0,Z=200.0),TraceEnd=(X=0.0,Y=0.0,Z=0.0),PenetrationDepth=0.0))
 
-+LineTraceSingleByChannelTests=(RootShapeAsset="/Engine/BasicShapes/Cube.Cube",HitResult=(Time=0.25,Location=(X=0.0,Y=0.0,Z=150.0),ImpactPoint=(X=0.0,Y=0.0,Z=50.0),Normal=(X=0.0,Y=-0.0,Z=1.0),ImpactNormal=(X=0.0,Y=0.0,Z=1.0),TraceStart=(X=0.0,Y=0.0,Z=200.0),TraceEnd=(X=0.0,Y=0.0,Z=0.0)))
-+LineTraceSingleByChannelTests=(RootShapeAsset="/Engine/BasicShapes/Cube.Cube",HitResult=(Time=0.50,Location=(X=-100.000008,Y=0.0,Z=0.0),ImpactPoint=(X=-50.0,Y=0.0,Z=0.0),Normal=(X=-1.0,Y=0.0,Z=0.0),ImpactNormal=(X=-1.0,Y=0.0,Z=0.0),TraceStart=(X=-200.0,Y=0.0,Z=0.0),TraceEnd=(X=0.0,Y=0.0,Z=0.0),PenetrationDepth=0.0)))
-+LineTraceSingleByChannelTests=(RootShapeAsset="/Engine/BasicShapes/Cube.Cube",HitResult=(Time=0.50,Location=(X=-100.000015,Y=0.0,Z=100.000015),ImpactPoint=(X=-50.0,Y=0.0,Z=50.0),Normal=(X=-1.0,Y=-0.0,Z=0.0),ImpactNormal=(X=0.0,Y=0.0,Z=1.0),TraceStart=(X=-200.0,Y=0.0,Z=200.0),TraceEnd=(X=0.0,Y=0.0,Z=0.0)))
++LineTraceSingleByChannelTests=(RootShapeAsset="/Engine/BasicShapes/Cube.Cube",HitResult=(Time=0.750,Location=(X=0.0,Y=0.0,Z=50.0),ImpactPoint=(X=0.0,Y=0.0,Z=50.0),Normal=(X=0.0,Y=0.0,Z=1.000000),ImpactNormal=(X=0.0,Y=0.0,Z=1.000000),TraceStart=(X=0.0,Y=0.0,Z=200.0),TraceEnd=(X=0.0,Y=0.0,Z=0.0),PenetrationDepth=0.0))
++LineTraceSingleByChannelTests=(RootShapeAsset="/Engine/BasicShapes/Cube.Cube",HitResult=(Time=0.750,Location=(X=-50.0,Y=0.0,Z=0.0),ImpactPoint=(X=-50.0,Y=0.0,Z=-0.0),Normal=(X=-1.000,Y=0.0,Z=0.0),ImpactNormal=(X=-1.000,Y=0.0,Z=0.0),TraceStart=(X=-200.0,Y=0.0,Z=0.0),TraceEnd=(X=0.0,Y=0.0,Z=0.0),PenetrationDepth=0.0))
++LineTraceSingleByChannelTests=(RootShapeAsset="/Engine/BasicShapes/Cube.Cube",HitResult=(Time=0.750,Location=(X=-50.000015,Y=0.0,Z=50.000015),ImpactPoint=(X=-50.0,Y=-0.0,Z=50.0),Normal=(X=0.0,Y=0.0,Z=1.000000),ImpactNormal=(X=0.0,Y=0.0,Z=1.000000),TraceStart=(X=-200.0,Y=0.0,Z=200.0),TraceEnd=(X=0.0,Y=0.0,Z=0.0),PenetrationDepth=0.0))
 
-+LineTraceSingleByChannelTests=(RootShapeAsset="/Engine/BasicShapes/Cylinder.Cylinder",HitResult=(Time=0.25,Location=(X=0.0,Y=0.0,Z=150.0),ImpactPoint=(X=0.0,Y=0.0,Z=50.0),Normal=(X=0.0,Y=-0.0,Z=1.0),ImpactNormal=(X=0.0,Y=0.0,Z=1.0),TraceStart=(X=0.0,Y=0.0,Z=200.0),TraceEnd=(X=0.0,Y=0.0,Z=0.0)))
-+LineTraceSingleByChannelTests=(RootShapeAsset="/Engine/BasicShapes/Cylinder.Cylinder",HitResult=(Time=0.50,Location=(X=-100.0,Y=0.0,Z=0.0),ImpactPoint=(X=-50.0,Y=0.000016,Z=30.0),Normal=(X=-1.0,Y=-0.0,Z=0.0),ImpactNormal=(X=-0.995185,Y=0.098017,Z=0.0),TraceStart=(X=-200.0,Y=0.0,Z=0.0),TraceEnd=(X=0.0,Y=0.0,Z=0.0)))
-+LineTraceSingleByChannelTests=(RootShapeAsset="/Engine/BasicShapes/Cylinder.Cylinder",HitResult=(Time=0.50,Location=(X=-100.000015,Y=0.0,Z=100.000015),ImpactPoint=(X=-50.0,Y=0.000016,Z=50.0),Normal=(X=-1.0,Y=-0.0,Z=0.0),ImpactNormal=(X=0.0,Y=0.0,Z=1.0),TraceStart=(X=-200.0,Y=0.0,Z=200.0),TraceEnd=(X=0.0,Y=0.0,Z=0.0)))
++LineTraceSingleByChannelTests=(RootShapeAsset="/Engine/BasicShapes/Cylinder.Cylinder",HitResult=(Time=0.750,Location=(X=0.0,Y=0.0,Z=50.000046),ImpactPoint=(X=0.0,Y=-0.0,Z=50.0),Normal=(X=0.0,Y=0.0,Z=1.000000),ImpactNormal=(X=0.0,Y=0.0,Z=1.000000),TraceStart=(X=0.0,Y=0.0,Z=200.0),TraceEnd=(X=0.0,Y=0.0,Z=0.0),PenetrationDepth=0.0))
++LineTraceSingleByChannelTests=(RootShapeAsset="/Engine/BasicShapes/Cylinder.Cylinder",HitResult=(Time=0.750,Location=(X=-50.0,Y=0.0,Z=0.0),ImpactPoint=(X=-50.0,Y=-0.0,Z=-0.0),Normal=(X=-0.995185,Y=-0.098017,Z=0.0),ImpactNormal=(X=-0.995185,Y=-0.098017,Z=0.0),TraceStart=(X=-200.0,Y=0.0,Z=0.0),TraceEnd=(X=0.0,Y=0.0,Z=0.0),PenetrationDepth=0.0))
++LineTraceSingleByChannelTests=(RootShapeAsset="/Engine/BasicShapes/Cylinder.Cylinder",HitResult=(Time=0.750,Location=(X=-50.000015,Y=0.0,Z=50.000015),ImpactPoint=(X=-50.0,Y=0.000006,Z=50.0),Normal=(X=0.0,Y=0.0,Z=1.000000),ImpactNormal=(X=0.0,Y=0.0,Z=1.000000),TraceStart=(X=-200.0,Y=0.0,Z=200.0),TraceEnd=(X=0.0,Y=0.0,Z=0.0),PenetrationDepth=0.0))
 
-+LineTraceSingleByChannelTests=(RootShapeAsset="/Engine/BasicShapes/Sphere.Sphere",HitResult=(Time=0.25,Location=(X=0.0,Y=0.0,Z=150.0),ImpactPoint=(X=0.0,Y=0.000002,Z=50.0),Normal=(X=0.0,Y=-0.0,Z=1.0),ImpactNormal=(X=0.028589,Y=0.094246,Z=0.995138),TraceStart=(X=0.0,Y=0.0,Z=200.0),TraceEnd=(X=0.0,Y=0.0,Z=0.0)))
-+LineTraceSingleByChannelTests=(RootShapeAsset="/Engine/BasicShapes/Sphere.Sphere",HitResult=(Time=0.50,Location=(X=-100.0,Y=0.0,Z=0.0),ImpactPoint=(X=-50.0,Y=0.000018,Z=0.000004),Normal=(X=-1.0,Y=-0.0,Z=0.0),ImpactNormal=(X=-0.990438,Y=0.097550,Z=-0.097550),TraceStart=(X=-200.0,Y=0.0,Z=0.0),TraceEnd=(X=0.0,Y=0.0,Z=0.0)))
-+LineTraceSingleByChannelTests=(RootShapeAsset="/Engine/BasicShapes/Sphere.Sphere",HitResult=(Time=0.544464,Location=(X=-91.107117,Y=0.0,Z=91.107117),ImpactPoint=(X=-46.193974,Y=0.000017,Z=19.134171),Normal=(X=-0.898263,Y=-0.0,Z=0.439459),ImpactNormal=(X=-0.878613,Y=-0.086536,Z=0.469628),TraceStart=(X=-200.0,Y=0.0,Z=200.0),TraceEnd=(X=0.0,Y=0.0,Z=0.0)))
++LineTraceSingleByChannelTests=(RootShapeAsset="/Engine/BasicShapes/Sphere.Sphere",HitResult=(Time=0.750,Location=(X=0.0,Y=0.0,Z=50.000015),ImpactPoint=(X=-0.0,Y=0.0,Z=50.0),Normal=(X=0.028589,Y=0.094246,Z=0.995138),ImpactNormal=(X=0.028589,Y=0.094246,Z=0.995138),TraceStart=(X=0.0,Y=0.0,Z=200.0),TraceEnd=(X=0.0,Y=0.0,Z=0.0),PenetrationDepth=0.0))
++LineTraceSingleByChannelTests=(RootShapeAsset="/Engine/BasicShapes/Sphere.Sphere",HitResult=(Time=0.750,Location=(X=-50.0,Y=0.0,Z=0.0),ImpactPoint=(X=-49.999996,Y=0.000001,Z=0.0),Normal=(X=-0.990438,Y=-0.097550,Z=-0.097550),ImpactNormal=(X=-0.990438,Y=-0.097550,Z=-0.097550),TraceStart=(X=-200.0,Y=0.0,Z=0.0),TraceEnd=(X=0.0,Y=0.0,Z=0.0),PenetrationDepth=0.0))
++LineTraceSingleByChannelTests=(RootShapeAsset="/Engine/BasicShapes/Sphere.Sphere",HitResult=(Time=0.823223,Location=(X=-35.355316,Y=0.0,Z=35.355316),ImpactPoint=(X=-35.355335,Y=-0.000003,Z=35.355339),Normal=(X=-0.770780,Y=-0.075915,Z=0.632563),ImpactNormal=(X=-0.770780,Y=-0.075915,Z=0.632563),TraceStart=(X=-200.0,Y=0.0,Z=200.0),TraceEnd=(X=0.0,Y=0.0,Z=0.0),PenetrationDepth=0.0))
 
 
 
