hooked up new bloom/gaussianblur setting to BaseScalability.ini

diff --git a/Engine/Config/BaseScalability.ini b/Engine/Config/BaseScalability.ini
--- a/Engine/Config/BaseScalability.ini
+++ b/Engine/Config/BaseScalability.ini
@@ -115,6 +115,7 @@ r.FastBlurThreshold=0
 r.Upscale.Quality=1
 r.TonemapperQuality=0
 r.LightShaftQuality=0
+r.Filter.SizeScale=0.6
 
 [PostProcessQuality@1]
 r.MotionBlurQuality=3
@@ -132,6 +133,7 @@ r.FastBlurThreshold=2
 r.Upscale.Quality=2
 r.TonemapperQuality=0
 r.LightShaftQuality=0
+r.Filter.SizeScale=0.7
 
 [PostProcessQuality@2]
 r.MotionBlurQuality=3
@@ -149,6 +151,7 @@ r.FastBlurThreshold=3
 r.Upscale.Quality=2
 r.TonemapperQuality=1
 r.LightShaftQuality=1
+r.Filter.SizeScale=0.8
 
 [PostProcessQuality@3]
 r.MotionBlurQuality=4
@@ -166,6 +169,7 @@ r.FastBlurThreshold=100
 r.Upscale.Quality=3
 r.TonemapperQuality=1
 r.LightShaftQuality=1
+r.Filter.SizeScale=1
 
 ;-----------------------------------------------------------------------------------------------------------------
 
