Paper2D: Added a redirector for PaperSpriteSheet to cover the fact that PaperJsonImporter was renamed to PaperSpritesheetImporter
[UE-16918]

diff --git a/Engine/Config/BaseEngine.ini b/Engine/Config/BaseEngine.ini
--- a/Engine/Config/BaseEngine.ini
+++ b/Engine/Config/BaseEngine.ini
@@ -624,6 +624,9 @@ bUseStreamingPause=false
 +TaggedPropertyRedirects=(ClassName="SpriteGeometryCollection",OldPropertyName="Polygons",NewPropertyName="Shapes")
 +TaggedPropertyRedirects=(ClassName="PaperTileSet",OldPropertyName="ExperimentalPerTileData",NewPropertyName="PerTileData")
 
+; Paper2D modifications (4.8.1 era)
++ActiveClassRedirects=(OldClassName="PaperSpriteSheet",NewClassName="/Script/PaperSpritesheetImporter.PaperSpriteSheet")
+
 ; collision response container refactor
 +TaggedPropertyRedirects=(ClassName="CollisionResponseContainer",OldPropertyName="RigidBody",NewPropertyName="PhysicsBody")
 +TaggedPropertyRedirects=(ClassName="CollisionResponseContainer",OldPropertyName="Static",NewPropertyName="WorldStatic")
