MovieScene: Added missing redirector

diff --git a/Engine/Config/BaseEngine.ini b/Engine/Config/BaseEngine.ini
--- a/Engine/Config/BaseEngine.ini
+++ b/Engine/Config/BaseEngine.ini
@@ -480,6 +480,7 @@ bUseStreamingPause=false
 
 ; MovieSceneCore moved to MovieScene
 +ActiveClassRedirects=(OldClassName="MovieScene",NewClassName="/Script/MovieScene.MovieScene")
++ActiveClassRedirects=(OldClassName="MovieSceneBoundObject",NewClassName="/Script/MovieScene.MovieSceneBoundObject")
 +ActiveClassRedirects=(OldClassName="MovieSceneBindings",NewClassName="/Script/MovieScene.MovieSceneBindings")
 +ActiveClassRedirects=(OldClassName="MovieScenePropertyTrack",NewClassName="/Script/MovieScene.MovieScenePropertyTrack")
 +ActiveClassRedirects=(OldClassName="MovieSceneSection",NewClassName="/Script/MovieScene.MovieSceneSection")
