MovieScene: Added more missing redirectors

diff --git a/Engine/Config/BaseEngine.ini b/Engine/Config/BaseEngine.ini
--- a/Engine/Config/BaseEngine.ini
+++ b/Engine/Config/BaseEngine.ini
@@ -480,13 +480,18 @@ bUseStreamingPause=false
 
 ; MovieSceneCore moved to MovieScene
 +ActiveClassRedirects=(OldClassName="MovieScene",NewClassName="/Script/MovieScene.MovieScene")
-+ActiveClassRedirects=(OldClassName="MovieSceneBoundObject",NewClassName="/Script/MovieScene.MovieSceneBoundObject")
 +ActiveClassRedirects=(OldClassName="MovieSceneBindings",NewClassName="/Script/MovieScene.MovieSceneBindings")
 +ActiveClassRedirects=(OldClassName="MovieScenePropertyTrack",NewClassName="/Script/MovieScene.MovieScenePropertyTrack")
 +ActiveClassRedirects=(OldClassName="MovieSceneSection",NewClassName="/Script/MovieScene.MovieSceneSection")
 +ActiveClassRedirects=(OldClassName="MovieSceneTrack",NewClassName="/Script/MovieScene.MovieSceneTrack")
 +ActiveClassRedirects=(OldClassName="RuntimeMovieScenePlayer",NewClassName="/Script/MovieScene.RuntimeMovieScenePlayer")
 
++ActiveStructRedirects=(OldStructName="MovieSceneSpawnable",NewStructName="/Script/MovieScene.MovieSceneSpawnable")
++ActiveStructRedirects=(OldStructName="MovieScenePossessable",NewStructName="/Script/MovieScene.MovieScenePossessable")
++ActiveStructRedirects=(OldStructName="MovieSceneEditorData",NewStructName="/Script/MovieScene.MovieSceneEditorData")
++ActiveStructRedirects=(OldStructName="MovieSceneObjectBinding",NewStructName="/Script/MovieScene.MovieSceneObjectBinding")
++ActiveStructRedirects=(OldStructName="MovieSceneBoundObject",NewStructName="/Script/MovieScene.MovieSceneBoundObject")
+
 ; MovieSceneCoreTypes moved to MovieSceneTracks
 +ActiveClassRedirects=(OldClassName="MovieScene3DTransformSection",NewClassName="/Script/MovieSceneTracks.MovieScene3DTransformSection")
 +ActiveClassRedirects=(OldClassName="MovieSceneAnimationSection",NewClassName="/Script/MovieSceneTracks.MovieSceneAnimationSection")
