MovieScene: Added package redirects

diff --git a/Engine/Config/BaseEngine.ini b/Engine/Config/BaseEngine.ini
--- a/Engine/Config/BaseEngine.ini
+++ b/Engine/Config/BaseEngine.ini
@@ -479,6 +479,8 @@ bUseStreamingPause=false
 +ActiveClassRedirects=(OldClassName="PawnSensingComponent",NewClassName="/Script/AIModule.PawnSensingComponent")
 
 ; MovieSceneCore moved to MovieScene
++ActiveGameNameRedirects=(OldGameName="/Script/MovieSceneCore", NewGameName="/Script/MovieScene")
+
 +ActiveClassRedirects=(OldClassName="MovieScene",NewClassName="/Script/MovieScene.MovieScene")
 +ActiveClassRedirects=(OldClassName="MovieSceneBindings",NewClassName="/Script/MovieScene.MovieSceneBindings")
 +ActiveClassRedirects=(OldClassName="MovieScenePropertyTrack",NewClassName="/Script/MovieScene.MovieScenePropertyTrack")
@@ -493,6 +495,8 @@ bUseStreamingPause=false
 +ActiveStructRedirects=(OldStructName="MovieSceneBoundObject",NewStructName="/Script/MovieScene.MovieSceneBoundObject")
 
 ; MovieSceneCoreTypes moved to MovieSceneTracks
++ActiveGameNameRedirects=(OldGameName="/Script/MovieSceneCoreTypes", NewGameName="/Script/MovieSceneTracks")
+
 +ActiveClassRedirects=(OldClassName="MovieScene3DTransformSection",NewClassName="/Script/MovieSceneTracks.MovieScene3DTransformSection")
 +ActiveClassRedirects=(OldClassName="MovieSceneAnimationSection",NewClassName="/Script/MovieSceneTracks.MovieSceneAnimationSection")
 +ActiveClassRedirects=(OldClassName="MovieSceneAudioSection",NewClassName="/Script/MovieSceneTracks.MovieSceneAudioSection")
