UMG - Updating the redirects to include moving FAnchors to the Slate module.

diff --git a/Engine/Config/BaseEngine.ini b/Engine/Config/BaseEngine.ini
--- a/Engine/Config/BaseEngine.ini
+++ b/Engine/Config/BaseEngine.ini
@@ -1231,6 +1231,10 @@ bUseStreamingPause=false
 +ActiveClassRedirects=(OldClassName="ESlateCheckBoxState", NewClassName="/Script/SlateCore.ECheckBoxState")
 +ActiveClassRedirects=(OldClassName="ECheckBoxState",NewClassName="/Script/SlateCore.ECheckBoxState")
 
+; UMG Moving FAnchor to Slate
++ActiveClassRedirects=(OldClassName="Anchors",NewClassName="/Script/Slate.Anchors")
++ActiveStructRedirects=(OldStructName="Anchors",NewStructName="/Script/Slate.Anchors")
+
 ; SplineComponent API overhaul
 +K2FieldRedirects=(OldFieldName="SplineComponent.GetNumSplinePoints",NewFieldName="SplineComponent.GetNumberOfSplinePoints")
 
