Change PathToRAC in DefaultEngine.ini from array type to non-array.

diff --git a/Engine/Config/BaseEngine.ini b/Engine/Config/BaseEngine.ini
--- a/Engine/Config/BaseEngine.ini
+++ b/Engine/Config/BaseEngine.ini
@@ -1918,7 +1918,7 @@ Foliage=(NumLODs=1,Name=LOCTEXT("FoliageLOD","Foliage"))
 ; +BucketConfigs=(Name="Advertisement", Size=4000000)
 [RuntimeAssetCache]
 +BucketConfigs=(Name="DefaultBucket", Size=10000000)
-+PathToRAC=RuntimeAssetCache
+PathToRAC=RuntimeAssetCache
 
 ; Derived Data backend graphs
 ; Each of the below backend graph sections contains a set of nodes used to create derived data backed graph.
