PR #1130: Rename add method for consistency. (Contributed by [REDACTED_PERSON])

diff --git a/Engine/Config/BaseEngine.ini b/Engine/Config/BaseEngine.ini
--- a/Engine/Config/BaseEngine.ini
+++ b/Engine/Config/BaseEngine.ini
@@ -1235,6 +1235,9 @@ bUseStreamingPause=false
 +ActiveClassRedirects=(OldClassName="Anchors",NewClassName="/Script/Slate.Anchors")
 +ActiveStructRedirects=(OldStructName="Anchors",NewStructName="/Script/Slate.Anchors")
 
+; UMG VerticalBox AddChild method rename
++K2FieldRedirects=(OldFieldName="VerticalBox.AddChildVerticalBox", NewFieldName="VerticalBox.AddChildToVerticalBox")
+
 ; SplineComponent API overhaul
 +K2FieldRedirects=(OldFieldName="SplineComponent.GetNumSplinePoints",NewFieldName="SplineComponent.GetNumberOfSplinePoints")
 
diff --git a/Engine/Source/Runtime/UMG/Private/Components/VerticalBox.cpp b/Engine/Source/Runtime/UMG/Private/Components/VerticalBox.cpp
--- a/Engine/Source/Runtime/UMG/Private/Components/VerticalBox.cpp
+++ b/Engine/Source/Runtime/UMG/Private/Components/VerticalBox.cpp
@@ -50,7 +50,7 @@ void UVerticalBox::OnSlotRemoved(UPanelSlot* Slot)
 	}
 }
 
-UVerticalBoxSlot* UVerticalBox::AddChildVerticalBox(UWidget* Content)
+UVerticalBoxSlot* UVerticalBox::AddChildToVerticalBox(UWidget* Content)
 {
 	return Cast<UVerticalBoxSlot>(Super::AddChild(Content));
 }
diff --git a/Engine/Source/Runtime/UMG/Public/Components/VerticalBox.h b/Engine/Source/Runtime/UMG/Public/Components/VerticalBox.h
--- a/Engine/Source/Runtime/UMG/Public/Components/VerticalBox.h
+++ b/Engine/Source/Runtime/UMG/Public/Components/VerticalBox.h
@@ -22,7 +22,7 @@ public:
 
 	/**  */
 	UFUNCTION(BlueprintCallable, Category="Panel")
-	UVerticalBoxSlot* AddChildVerticalBox(UWidget* Content);
+	UVerticalBoxSlot* AddChildToVerticalBox(UWidget* Content);
 
 #if WITH_EDITOR
 	// UWidget interface
