Engine: Added a 4.9 note to some redirectors in BaseEngine.ini

diff --git a/Engine/Config/BaseEngine.ini b/Engine/Config/BaseEngine.ini
--- a/Engine/Config/BaseEngine.ini
+++ b/Engine/Config/BaseEngine.ini
@@ -382,7 +382,7 @@ bUseStreamingPause=false
 +ActiveClassRedirects=(OldClassName="AnimNode_TwoBoneIK",NewClassName="/Script/AnimGraphRuntime.AnimNode_TwoBoneIK")
 +ActiveClassRedirects=(OldClassName="AnimNode_WheelHandler",NewClassName="/Script/AnimGraphRuntime.AnimNode_WheelHandler")
 
-; Renamed all user-exposed locations for rotator fields to display as XYZ (RPY) instead of YPR
+; Renamed all user-exposed locations for rotator fields to display as XYZ (RPY) instead of YPR (4.9 era)
 +K2FieldRedirects=(OldFieldName="KismetMathLibrary.BreakRot", NewFieldName="KismetMathLibrary.BreakRotator")
 +K2FieldRedirects=(OldFieldName="KismetMathLibrary.MakeRot", NewFieldName="KismetMathLibrary.MakeRotator")
 
