- Fixed vehicle advanced template crash
- Fixed two way blend node redirects
- Added struct redirects for all classes.

diff --git a/Engine/Config/BaseEngine.ini b/Engine/Config/BaseEngine.ini
--- a/Engine/Config/BaseEngine.ini
+++ b/Engine/Config/BaseEngine.ini
@@ -382,6 +382,11 @@ bUseStreamingPause=false
 +ActiveClassRedirects=(OldClassName="AnimNode_TwoBoneIK",NewClassName="/Script/AnimGraphRuntime.AnimNode_TwoBoneIK")
 +ActiveClassRedirects=(OldClassName="AnimNode_WheelHandler",NewClassName="/Script/AnimGraphRuntime.AnimNode_WheelHandler")
 
+
++ActiveClassRedirects=(OldClassName="AnimGraphNode_BlendSpace",NewClassName="/Script/AnimGraph.AnimGraphNode_BlendSpacePlayer")
++ActiveClassRedirects=(OldClassName="AnimNode_BlendSpace",NewClassName="/Script/Engine.AnimNode_BlendSpacePlayer")
++ActiveStructRedirects=(OldStructName="AnimNode_BlendSpace",NewStructName="AnimNode_BlendSpacePlayer")
+
 ; Moving more animation nodes form Engine to AnimGraphRuntime
 +ActiveClassRedirects=(OldClassName="AnimNode_ApplyAdditive",NewClassName="/Script/AnimGraphRuntime.AnimNode_ApplyAdditive")
 +ActiveClassRedirects=(OldClassName="AnimNode_BlendListBase",NewClassName="/Script/AnimGraphRuntime.AnimNode_BlendListBase")
@@ -392,13 +397,33 @@ bUseStreamingPause=false
 +ActiveClassRedirects=(OldClassName="AnimNode_BlendSpacePlayer",NewClassName="/Script/AnimGraphRuntime.AnimNode_BlendSpacePlayer")
 +ActiveClassRedirects=(OldClassName="AnimNode_LayeredBoneBlend",NewClassName="/Script/AnimGraphRuntime.AnimNode_LayeredBoneBlend")
 +ActiveClassRedirects=(OldClassName="AnimNode_RefPose",NewClassName="/Script/AnimGraphRuntime.AnimNode_RefPose")
++ActiveClassRedirects=(OldClassName="AnimNode_MeshSpaceRefPose",NewClassName="/Script/AnimGraphRuntime.AnimNode_MeshSpaceRefPose")
 +ActiveClassRedirects=(OldClassName="AnimNode_Root",NewClassName="/Script/AnimGraphRuntime.AnimNode_Root")
 +ActiveClassRedirects=(OldClassName="AnimNode_RotateRootBone",NewClassName="/Script/AnimGraphRuntime.AnimNode_RotateRootBone")
 +ActiveClassRedirects=(OldClassName="AnimNode_RotationOffsetBlendSpace",NewClassName="/Script/AnimGraphRuntime.AnimNode_RotationOffsetBlendSpace")
 +ActiveClassRedirects=(OldClassName="AnimNode_SaveCachedPose",NewClassName="/Script/AnimGraphRuntime.AnimNode_SaveCachedPose")
 +ActiveClassRedirects=(OldClassName="AnimNode_SequenceEvaluator",NewClassName="/Script/AnimGraphRuntime.AnimNode_SequenceEvaluator")
 +ActiveClassRedirects=(OldClassName="AnimNode_Slot",NewClassName="/Script/AnimGraphRuntime.AnimNode_Slot")
-+ActiveClassRedirects=(OldClassName="AnimNode_TwoWayBlend",NewClassName="/Script/AnimGraphRuntime.AnimNode_TwoWayBlend")
++ActiveClassRedirects=(OldClassName="AnimationNode_TwoWayBlend",NewClassName="/Script/AnimGraphRuntime.AnimationNode_TwoWayBlend")
+
+; Moving more animation nodes form Engine to AnimGraphRuntime
++ActiveStructRedirects=(OldStructName="AnimNode_ApplyAdditive",NewStructName="/Script/AnimGraphRuntime.AnimNode_ApplyAdditive")
++ActiveStructRedirects=(OldStructName="AnimNode_BlendListBase",NewStructName="/Script/AnimGraphRuntime.AnimNode_BlendListBase")
++ActiveStructRedirects=(OldStructName="AnimNode_BlendListByBool",NewStructName="/Script/AnimGraphRuntime.AnimNode_BlendListByBool")
++ActiveStructRedirects=(OldStructName="AnimNode_BlendListByEnum",NewStructName="/Script/AnimGraphRuntime.AnimNode_BlendListByEnum")
++ActiveStructRedirects=(OldStructName="AnimNode_BlendListByInt",NewStructName="/Script/AnimGraphRuntime.AnimNode_BlendListByInt")
++ActiveStructRedirects=(OldStructName="AnimNode_BlendSpaceEvaluator",NewStructName="/Script/AnimGraphRuntime.AnimNode_BlendSpaceEvaluator")
++ActiveStructRedirects=(OldStructName="AnimNode_BlendSpacePlayer",NewStructName="/Script/AnimGraphRuntime.AnimNode_BlendSpacePlayer")
++ActiveStructRedirects=(OldStructName="AnimNode_LayeredBoneBlend",NewStructName="/Script/AnimGraphRuntime.AnimNode_LayeredBoneBlend")
++ActiveStructRedirects=(OldStructName="AnimNode_RefPose",NewStructName="/Script/AnimGraphRuntime.AnimNode_RefPose")
++ActiveStructRedirects=(OldStructName="AnimNode_MeshSpaceRefPose",NewStructName="/Script/AnimGraphRuntime.AnimNode_MeshSpaceRefPose")
++ActiveStructRedirects=(OldStructName="AnimNode_Root",NewStructName="/Script/AnimGraphRuntime.AnimNode_Root")
++ActiveStructRedirects=(OldStructName="AnimNode_RotateRootBone",NewStructName="/Script/AnimGraphRuntime.AnimNode_RotateRootBone")
++ActiveStructRedirects=(OldStructName="AnimNode_RotationOffsetBlendSpace",NewStructName="/Script/AnimGraphRuntime.AnimNode_RotationOffsetBlendSpace")
++ActiveStructRedirects=(OldStructName="AnimNode_SaveCachedPose",NewStructName="/Script/AnimGraphRuntime.AnimNode_SaveCachedPose")
++ActiveStructRedirects=(OldStructName="AnimNode_SequenceEvaluator",NewStructName="/Script/AnimGraphRuntime.AnimNode_SequenceEvaluator")
++ActiveStructRedirects=(OldStructName="AnimNode_Slot",NewStructName="/Script/AnimGraphRuntime.AnimNode_Slot")
++ActiveStructRedirects=(OldStructName="AnimationNode_TwoWayBlend",NewStructName="/Script/AnimGraphRuntime.AnimationNode_TwoWayBlend")
 
 ; Renamed all user-exposed locations for rotator fields to display as XYZ (RPY) instead of YPR (4.9 era)
 +K2FieldRedirects=(OldFieldName="KismetMathLibrary.BreakRot", NewFieldName="KismetMathLibrary.BreakRotator")
@@ -1112,9 +1137,6 @@ bUseStreamingPause=false
 +K2ParamRedirects=(NodeName="/Script/BlueprintGraph.K2Node_CallFunction", OldParamName="KismetMathLibrary.ProjectVectorOnToVector.X", NewParamName="V")
 +K2ParamRedirects=(NodeName="/Script/BlueprintGraph.K2Node_CallFunction", OldParamName="KismetMathLibrary.ProjectVectorOnToVector.Y", NewParamName="Target")
 
-+K2ParamRedirects=(NodeName="/Script/BlueprintGraph.K2Node_CallFunction", OldParamName="KismetMathLibrary.ConvertTransformToRelative.WorldTransform", NewParamName="Transform")
-+K2ParamRedirects=(NodeName="/Script/BlueprintGraph.K2Node_CallFunction", OldParamName="KismetMathLibrary.ConvertTransformToRelative.LocalTransform", NewParamName="ParentTransform")
-
 +K2FieldRedirects=(OldFieldName="PrimitiveComponent.SetRBLinearVelocity",NewFieldName="SetPhysicsLinearVelocity")
 +K2FieldRedirects=(OldFieldName="PrimitiveComponent.GetRBLinearVelocity",NewFieldName="GetPhysicsLinearVelocity")
 +K2FieldRedirects=(OldFieldName="PrimitiveComponent.SetAllRBLinearVelocity",NewFieldName="SetAllPhysicsLinearVelocity")
@@ -1174,10 +1196,6 @@ bUseStreamingPause=false
 +TaggedPropertyRedirects=(ClassName="GameInfo",OldPropertyName="PlayerReplicationInfoClass",NewPropertyName="PlayerStateClass")
 +TaggedPropertyRedirects=(ClassName="GameMode",OldPropertyName="PlayerReplicationInfoClass",NewPropertyName="PlayerStateClass")
 
-+ActiveClassRedirects=(OldClassName="AnimGraphNode_BlendSpace",NewClassName="/Script/AnimGraph.AnimGraphNode_BlendSpacePlayer")
-+ActiveClassRedirects=(OldClassName="AnimNode_BlendSpace",NewClassName="/Script/Engine.AnimNode_BlendSpacePlayer")
-+ActiveStructRedirects=(OldStructName="AnimNode_BlendSpace",NewStructName="AnimNode_BlendSpacePlayer")
-
 +TaggedPropertyRedirects=(ClassName="GameState",OldPropertyName="GameClass",NewPropertyName="GameModeClass")
 +TaggedPropertyRedirects=(ClassName="GameReplicationInfo",OldPropertyName="GameClass",NewPropertyName="GameModeClass")
 
@@ -2067,16 +2085,6 @@ AspectRatioAxisConstraint=AspectRatio_MaintainXFOV
 +Class=ParticleSystem
 +Class=Texture2D
 
-[AssetRegistry]
-; Fill out this list with asset registry tags that are needed in cooked builds. All other tags are stripped during cook.
-; You can use a * wildcard to refer to either all classes or all tags for the class.
-; If your game needs all asset registry tags, simply add (Class=*,Tag=*). Be warned that this will include editor only data and it is advised you specify which tags you use.
-; ObjectLibraries
-+CookedTagsWhitelist=(Class=Blueprint,Tag=ParentClass)
-; GameplayCueManager
-+CookedTagsWhitelist=(Class=Blueprint,Tag=GeneratedClass)
-+CookedTagsWhitelist=(Class=Blueprint,Tag=GameplayCueName)
-
 [AutomationTesting]
 ImportTestPath=../../Content/EditorAutomation/
 ImportTestPackagePath=/Engine/Content/EditorAutomation
