Fix build error

diff --git a/Engine/Config/BaseEngine.ini b/Engine/Config/BaseEngine.ini
--- a/Engine/Config/BaseEngine.ini
+++ b/Engine/Config/BaseEngine.ini
@@ -384,7 +384,6 @@ bUseStreamingPause=false
 
 
 +ActiveClassRedirects=(OldClassName="AnimGraphNode_BlendSpace",NewClassName="/Script/AnimGraph.AnimGraphNode_BlendSpacePlayer")
-+ActiveClassRedirects=(OldClassName="AnimNode_BlendSpace",NewClassName="/Script/Engine.AnimNode_BlendSpacePlayer")
 +ActiveStructRedirects=(OldStructName="AnimNode_BlendSpace",NewStructName="AnimNode_BlendSpacePlayer")
 
 ; Moving more animation nodes form Engine to AnimGraphRuntime
@@ -394,6 +393,7 @@ bUseStreamingPause=false
 +ActiveClassRedirects=(OldClassName="AnimNode_BlendListByEnum",NewClassName="/Script/AnimGraphRuntime.AnimNode_BlendListByEnum")
 +ActiveClassRedirects=(OldClassName="AnimNode_BlendListByInt",NewClassName="/Script/AnimGraphRuntime.AnimNode_BlendListByInt")
 +ActiveClassRedirects=(OldClassName="AnimNode_BlendSpaceEvaluator",NewClassName="/Script/AnimGraphRuntime.AnimNode_BlendSpaceEvaluator")
++ActiveClassRedirects=(OldClassName="AnimNode_BlendSpace",NewClassName="/Script/AnimGraphRuntime.AnimNode_BlendSpacePlayer")
 +ActiveClassRedirects=(OldClassName="AnimNode_BlendSpacePlayer",NewClassName="/Script/AnimGraphRuntime.AnimNode_BlendSpacePlayer")
 +ActiveClassRedirects=(OldClassName="AnimNode_LayeredBoneBlend",NewClassName="/Script/AnimGraphRuntime.AnimNode_LayeredBoneBlend")
 +ActiveClassRedirects=(OldClassName="AnimNode_RefPose",NewClassName="/Script/AnimGraphRuntime.AnimNode_RefPose")
@@ -406,24 +406,6 @@ bUseStreamingPause=false
 +ActiveClassRedirects=(OldClassName="AnimNode_Slot",NewClassName="/Script/AnimGraphRuntime.AnimNode_Slot")
 +ActiveClassRedirects=(OldClassName="AnimationNode_TwoWayBlend",NewClassName="/Script/AnimGraphRuntime.AnimationNode_TwoWayBlend")
 
-; Moving more animation nodes form Engine to AnimGraphRuntime
-+ActiveStructRedirects=(OldStructName="AnimNode_ApplyAdditive",NewStructName="/Script/AnimGraphRuntime.AnimNode_ApplyAdditive")
-+ActiveStructRedirects=(OldStructName="AnimNode_BlendListBase",NewStructName="/Script/AnimGraphRuntime.AnimNode_BlendListBase")
-+ActiveStructRedirects=(OldStructName="AnimNode_BlendListByBool",NewStructName="/Script/AnimGraphRuntime.AnimNode_BlendListByBool")
-+ActiveStructRedirects=(OldStructName="AnimNode_BlendListByEnum",NewStructName="/Script/AnimGraphRuntime.AnimNode_BlendListByEnum")
-+ActiveStructRedirects=(OldStructName="AnimNode_BlendListByInt",NewStructName="/Script/AnimGraphRuntime.AnimNode_BlendListByInt")
-+ActiveStructRedirects=(OldStructName="AnimNode_BlendSpaceEvaluator",NewStructName="/Script/AnimGraphRuntime.AnimNode_BlendSpaceEvaluator")
-+ActiveStructRedirects=(OldStructName="AnimNode_BlendSpacePlayer",NewStructName="/Script/AnimGraphRuntime.AnimNode_BlendSpacePlayer")
-+ActiveStructRedirects=(OldStructName="AnimNode_LayeredBoneBlend",NewStructName="/Script/AnimGraphRuntime.AnimNode_LayeredBoneBlend")
-+ActiveStructRedirects=(OldStructName="AnimNode_RefPose",NewStructName="/Script/AnimGraphRuntime.AnimNode_RefPose")
-+ActiveStructRedirects=(OldStructName="AnimNode_MeshSpaceRefPose",NewStructName="/Script/AnimGraphRuntime.AnimNode_MeshSpaceRefPose")
-+ActiveStructRedirects=(OldStructName="AnimNode_Root",NewStructName="/Script/AnimGraphRuntime.AnimNode_Root")
-+ActiveStructRedirects=(OldStructName="AnimNode_RotateRootBone",NewStructName="/Script/AnimGraphRuntime.AnimNode_RotateRootBone")
-+ActiveStructRedirects=(OldStructName="AnimNode_RotationOffsetBlendSpace",NewStructName="/Script/AnimGraphRuntime.AnimNode_RotationOffsetBlendSpace")
-+ActiveStructRedirects=(OldStructName="AnimNode_SaveCachedPose",NewStructName="/Script/AnimGraphRuntime.AnimNode_SaveCachedPose")
-+ActiveStructRedirects=(OldStructName="AnimNode_SequenceEvaluator",NewStructName="/Script/AnimGraphRuntime.AnimNode_SequenceEvaluator")
-+ActiveStructRedirects=(OldStructName="AnimNode_Slot",NewStructName="/Script/AnimGraphRuntime.AnimNode_Slot")
-+ActiveStructRedirects=(OldStructName="AnimationNode_TwoWayBlend",NewStructName="/Script/AnimGraphRuntime.AnimationNode_TwoWayBlend")
 
 ; Renamed all user-exposed locations for rotator fields to display as XYZ (RPY) instead of YPR (4.9 era)
 +K2FieldRedirects=(OldFieldName="KismetMathLibrary.BreakRot", NewFieldName="KismetMathLibrary.BreakRotator")
