Made sure MobileHDR invalidates iterative cooked content.

diff --git a/Engine/Config/BaseEditor.ini b/Engine/Config/BaseEditor.ini
--- a/Engine/Config/BaseEditor.ini
+++ b/Engine/Config/BaseEditor.ini
@@ -429,7 +429,7 @@ bUseFullGCAssetClassNames=true
 +VersionedIntRValues=r.CompileShadersForDevelopment
 +VersionedIntRValues=r.PS4MixedModeShaderDebugInfo
 +VersionedIntRValues=r.PS4DumpShaderSDB
-
++VersionedIntRValues=r.MobileHDR
 
 
 [/Script/Localization.LocalizationSettings]
