Engine: Added an era in the comment for a set of redirectors in BaseEngine.ini

diff --git a/Engine/Config/BaseEngine.ini b/Engine/Config/BaseEngine.ini
--- a/Engine/Config/BaseEngine.ini
+++ b/Engine/Config/BaseEngine.ini
@@ -386,7 +386,7 @@ bUseStreamingPause=false
 +ActiveClassRedirects=(OldClassName="AnimGraphNode_BlendSpace",NewClassName="/Script/AnimGraph.AnimGraphNode_BlendSpacePlayer")
 +ActiveStructRedirects=(OldStructName="AnimNode_BlendSpace",NewStructName="AnimNode_BlendSpacePlayer")
 
-; Moving more animation nodes form Engine to AnimGraphRuntime
+; Moving more animation nodes from Engine to AnimGraphRuntime (4.10 era)
 +ActiveClassRedirects=(OldClassName="AnimNode_ApplyAdditive",NewClassName="/Script/AnimGraphRuntime.AnimNode_ApplyAdditive")
 +ActiveClassRedirects=(OldClassName="AnimNode_BlendListBase",NewClassName="/Script/AnimGraphRuntime.AnimNode_BlendListBase")
 +ActiveClassRedirects=(OldClassName="AnimNode_BlendListByBool",NewClassName="/Script/AnimGraphRuntime.AnimNode_BlendListByBool")
