Added more ini settings which invalidate cooked content (incase we want to enable this code path in the future).

diff --git a/Engine/Config/BaseEditor.ini b/Engine/Config/BaseEditor.ini
--- a/Engine/Config/BaseEditor.ini
+++ b/Engine/Config/BaseEditor.ini
@@ -422,6 +422,16 @@ bUseFullGCAssetClassNames=true
 +VersionedIniParams=Engine:/Script/Engine.Engine:WeightMapPlaceholderTextureName
 +VersionedIniParams=Engine:/Script/Engine.Engine:LightMapDensityTextureName
 
++VersionedIniParams=Game:/Script/UnrealEd.ProjectPackagingSettings:MapsToCook
++VersionedIniParams=Game:/Script/UnrealEd.ProjectPackagingSettings:DirectoriesToAlwaysCook
++VersionedIniParams=Game:/Script/UnrealEd.ProjectPackagingSettings:DirectoriesToNeverCook
++VersionedIniParams=Game:/Script/UnrealEd.ProjectPackagingSettings:bSkipEditorContent
++VersionedIniParams=Game:/Script/UnrealEd.ProjectPackagingSettings:bCompressed
++VersionedIniParams=Game:/Script/UnrealEd.ProjectPackagingSettings:bCookMapsOnly
++VersionedIniParams=Game:/Script/UnrealEd.ProjectPackagingSettings:bCookAll
++VersionedIniParams=Game:/Script/UnrealEd.ProjectPackagingSettings:DefaultCulture
++VersionedIniParams=Game:/Script/UnrealEd.ProjectPackagingSettings:UsePakFile
++VersionedIniParams=Game:/Script/UnrealEd.ProjectPackagingSettings:bGenerateChunks
 
 +VersionedIntRValues=r.AllowStaticLighting
 +VersionedIntRValues=r.GBuffer
