Point Persona tutorial to correct asset to fix cooking warning

diff --git a/Engine/Config/BaseEditorSettings.ini b/Engine/Config/BaseEditorSettings.ini
--- a/Engine/Config/BaseEditorSettings.ini
+++ b/Engine/Config/BaseEditorSettings.ini
@@ -38,7 +38,7 @@ bCaptureAudioFromMicrophone=true
 +TutorialContexts=(Context="DestructibleMeshEditor",BrowserFilter=,AttractTutorial=None,LaunchTutorial=/Engine/Tutorial/SubEditors/DestructibleMeshEditorTutorial.DestructibleMeshEditorTutorial_C)
 +TutorialContexts=(Context="MaterialEditor",BrowserFilter=,AttractTutorial=None,LaunchTutorial=/Engine/Tutorial/SubEditors/MaterialEditorTutorial.MaterialEditorTutorial_C)
 +TutorialContexts=(Context="Cascade",BrowserFilter=,AttractTutorial=None,LaunchTutorial=/Engine/Tutorial/SubEditors/ParticleSystemEditorTutorial.ParticleSystemEditorTutorial_C)
-+TutorialContexts=(Context="Persona",BrowserFilter=,AttractTutorial=None,LaunchTutorial=/Engine/Tutorial/SubEditors/PersonaAnimEditorTutorial.PersonaAnimEditorTutorial_C)
++TutorialContexts=(Context="Persona",BrowserFilter=,AttractTutorial=None,LaunchTutorial=/Engine/Tutorial/SubEditors/PersonaAnimEditorWalkThrough.PersonaAnimEditorWalkthrough_C)
 +TutorialContexts=(Context="SpriteEditor",BrowserFilter=,AttractTutorial=None,LaunchTutorial=/Engine/Tutorial/Paper2D/SpriteEditorTutorial.SpriteEditorTutorial_C)
 +TutorialContexts=(Context="FlipbookEditor",BrowserFilter=,AttractTutorial=None,LaunchTutorial=/Engine/Tutorial/Paper2D/FlipbookEditorTutorial.FlipbookEditorTutorial_C)
 +TutorialContexts=(Context="LandscapeMode",BrowserFilter=,AttractTutorial=None,LaunchTutorial=/Engine/Tutorial/Landscape/Landscape_Manage_Mode.Landscape_Manage_Mode_C)
