added new autocomplete for easier development

diff --git a/Engine/Config/BaseInput.ini b/Engine/Config/BaseInput.ini
--- a/Engine/Config/BaseInput.ini
+++ b/Engine/Config/BaseInput.ini
@@ -162,3 +162,4 @@ MaxScrollbackSize=1024
 +ManualAutoCompleteList=(Command="Log ",Desc="Log <category> <level> (Change verbosity level for a log category)")
 +ManualAutoCompleteList=(Command="Log list",Desc="Log list <search string> (Search for log categories)")
 +ManualAutoCompleteList=(Command="Log reset",Desc="Log reset (Undo any changes to log verbosity)")
++ManualAutoCompleteList=(Command="ToggleDrawEvents",Desc="ToggleDrawEvents (to toggle the PIX annotations for better profiling data e.g. when using Razor)")
