removed old behavior tree object redirectors
#ue4 UE-20918

diff --git a/Engine/Config/BaseEngine.ini b/Engine/Config/BaseEngine.ini
--- a/Engine/Config/BaseEngine.ini
+++ b/Engine/Config/BaseEngine.ini
@@ -315,12 +315,6 @@ bUseStreamingPause=false
 +TaggedPropertyRedirects=(ClassName="SkeletalMeshLODInfo",OldPropertyName="DisplayFactor",NewPropertyName="ScreenSize")
 
 ;Behavior Tree editor redirects
-+ActiveClassRedirects=(OldClassName="BehaviorTreeNode",NewClassName="BehaviorTreeGraphNode")
-+ActiveClassRedirects=(OldClassName="BehaviorTreeNode_Composite",NewClassName="BehaviorTreeGraphNode_Composite")
-+ActiveClassRedirects=(OldClassName="BehaviorTreeNode_Decorator",NewClassName="BehaviorTreeGraphNode_Decorator")
-+ActiveClassRedirects=(OldClassName="BehaviorTreeNode_Root",NewClassName="BehaviorTreeGraphNode_Root")
-+ActiveClassRedirects=(OldClassName="BehaviorTreeNode_Task",NewClassName="BehaviorTreeGraphNode_Task")
-+ActiveClassRedirects=(OldClassName="BTTask_GoTo",NewClassName="BTTask_MoveTo")
 +ActiveClassRedirects=(OldClassName="BTTask_RunQuery",NewClassName="BTTask_RunEQSQuery")
 
 +ActiveClassRedirects=(OldClassName="MoveComponentAction",NewClassName="/Script/Engine/KismetSystemLibrary.MoveComponentAction")
