Android: Added Adreno 420 profile (uses Android_High)

diff --git a/Engine/Config/BaseDeviceProfiles.ini b/Engine/Config/BaseDeviceProfiles.ini
--- a/Engine/Config/BaseDeviceProfiles.ini
+++ b/Engine/Config/BaseDeviceProfiles.ini
@@ -177,6 +177,7 @@ MatchProfile=(Profile="Android_PowerVR540",Match=((SourceType=SRC_GpuFamily,Comp
 +MatchProfile=(Profile="Android_Adreno320",Match=((SourceType=SRC_GpuFamily,CompareType=CMP_Equal,MatchString="Adreno (TM) 320")))
 +MatchProfile=(Profile="Android_Adreno330_Ver53",Match=((SourceType=SRC_GpuFamily,CompareType=CMP_Equal,MatchString="Adreno (TM) 330"),(SourceType=SRC_GlVersion,CompareType=CMP_Regex,MatchString="ES\\ 3\\.0\\ V@([0-9]+)"),(SourceType=SRC_PreviousRegexMatch,CompareType=CMP_GreaterEqual,MatchString="53")))
 +MatchProfile=(Profile="Android_Adreno330",Match=((SourceType=SRC_GpuFamily,CompareType=CMP_Equal,MatchString="Adreno (TM) 330")))
++MatchProfile=(Profile="Android_Adreno420",Match=((SourceType=SRC_GpuFamily,CompareType=CMP_Equal,MatchString="Adreno (TM) 420")))
 +MatchProfile=(Profile="Android_TegraK1_GL4",Match=((SourceType=SRC_GpuFamily,CompareType=CMP_Equal,MatchString="NVIDIA Tegra"),(SourceType=SRC_GlVersion,CompareType=CMP_Regex,MatchString="^4\\.")))
 +MatchProfile=(Profile="Android_TegraK1",Match=((SourceType=SRC_GpuFamily,CompareType=CMP_Equal,MatchString="NVIDIA Tegra"),(SourceType=SRC_GlVersion,CompareType=CMP_Regex,MatchString="^OpenGL\\ ES 3\\.")))
 +MatchProfile=(Profile="Android_Tegra4",Match=((SourceType=SRC_GpuFamily,CompareType=CMP_Equal,MatchString="NVIDIA Tegra"),(SourceType=SRC_GlVersion,CompareType=CMP_Regex,MatchString="^OpenGL\\ ES 2\\.")))
@@ -258,6 +259,10 @@ DeviceType=Android
 BaseProfileName=Android_Adreno330
 +CVars=r.DisjointTimerQueries=1
 
+[Android_Adreno420 DeviceProfile]
+DeviceType=Android
+BaseProfileName=Android_High
+
 [Android_PowerVR540 DeviceProfile]
 DeviceType=Android
 BaseProfileName=Android
