Android device profile detection - replaced individual Adreno 2xx device profiles with a single one and made the Adreno 4xx profile matcher generic.

diff --git a/Engine/Config/BaseDeviceProfiles.ini b/Engine/Config/BaseDeviceProfiles.ini
--- a/Engine/Config/BaseDeviceProfiles.ini
+++ b/Engine/Config/BaseDeviceProfiles.ini
@@ -169,25 +169,21 @@ BaseProfileName=IOS
 
 [/Script/AndroidDeviceProfileSelector.AndroidDeviceProfileMatchingRules]
 MatchProfile=(Profile="Android_PowerVR540",Match=((SourceType=SRC_GpuFamily,CompareType=CMP_Equal,MatchString="PowerVR SGX 540")))
-+MatchProfile=(Profile="Android_Adreno200",Match=((SourceType=SRC_GpuFamily,CompareType=CMP_Equal,MatchString="Adreno (TM) 200")))
-+MatchProfile=(Profile="Android_Adreno203",Match=((SourceType=SRC_GpuFamily,CompareType=CMP_Equal,MatchString="Adreno (TM) 203")))
-+MatchProfile=(Profile="Android_Adreno205",Match=((SourceType=SRC_GpuFamily,CompareType=CMP_Equal,MatchString="Adreno (TM) 205")))
-+MatchProfile=(Profile="Android_Adreno220",Match=((SourceType=SRC_GpuFamily,CompareType=CMP_Equal,MatchString="Adreno (TM) 220")))
-+MatchProfile=(Profile="Android_Adreno225",Match=((SourceType=SRC_GpuFamily,CompareType=CMP_Equal,MatchString="Adreno (TM) 225")))
++MatchProfile=(Profile="Android_Adreno2xx",Match=((SourceType=SRC_GpuFamily,CompareType=CMP_Regex,MatchString="Adreno \\(TM\\) 2[0-9][0-9]")))
 +MatchProfile=(Profile="Android_Adreno320",Match=((SourceType=SRC_GpuFamily,CompareType=CMP_Equal,MatchString="Adreno (TM) 320")))
-+MatchProfile=(Profile="Android_Adreno330_Ver53",Match=((SourceType=SRC_GpuFamily,CompareType=CMP_Equal,MatchString="Adreno (TM) 330"),(SourceType=SRC_GlVersion,CompareType=CMP_Regex,MatchString="ES\\ 3\\.0\\ V@([0-9]+)"),(SourceType=SRC_PreviousRegexMatch,CompareType=CMP_GreaterEqual,MatchString="53")))
++MatchProfile=(Profile="Android_Adreno330_Ver53",Match=((SourceType=SRC_GpuFamily,CompareType=CMP_Equal,MatchString="Adreno (TM) 330"),(SourceType=SRC_GlVersion,CompareType=CMP_Regex,MatchString="ES 3\\.0 V@([0-9]+)"),(SourceType=SRC_PreviousRegexMatch,CompareType=CMP_GreaterEqual,MatchString="53")))
 +MatchProfile=(Profile="Android_Adreno330",Match=((SourceType=SRC_GpuFamily,CompareType=CMP_Equal,MatchString="Adreno (TM) 330")))
-+MatchProfile=(Profile="Android_Adreno420",Match=((SourceType=SRC_GpuFamily,CompareType=CMP_Equal,MatchString="Adreno (TM) 420")))
++MatchProfile=(Profile="Android_Adreno4xx",Match=((SourceType=SRC_GpuFamily,CompareType=CMP_Regex,MatchString="Adreno \\(TM\\) 4[0-9][0-9]")))
 +MatchProfile=(Profile="Android_TegraK1_GL4",Match=((SourceType=SRC_GpuFamily,CompareType=CMP_Equal,MatchString="NVIDIA Tegra"),(SourceType=SRC_GlVersion,CompareType=CMP_Regex,MatchString="^4\\.")))
-+MatchProfile=(Profile="Android_TegraK1",Match=((SourceType=SRC_GpuFamily,CompareType=CMP_Equal,MatchString="NVIDIA Tegra"),(SourceType=SRC_GlVersion,CompareType=CMP_Regex,MatchString="^OpenGL\\ ES 3\\.")))
-+MatchProfile=(Profile="Android_Tegra4",Match=((SourceType=SRC_GpuFamily,CompareType=CMP_Equal,MatchString="NVIDIA Tegra"),(SourceType=SRC_GlVersion,CompareType=CMP_Regex,MatchString="^OpenGL\\ ES 2\\.")))
-+MatchProfile=(Profile="Android_IntelHD_ES3",Match=((SourceType=SRC_GpuFamily,CompareType=CMP_Equal,MatchString="Intel(R) HD Graphics"),(SourceType=SRC_GlVersion,CompareType=CMP_Regex,MatchString="^OpenGL\\ ES 3\\.")))
-+MatchProfile=(Profile="Android_IntelHD",Match=((SourceType=SRC_GpuFamily,CompareType=CMP_Equal,MatchString="Intel(R) HD Graphics"),(SourceType=SRC_GlVersion,CompareType=CMP_Regex,MatchString="^OpenGL\\ ES 2\\.")))
++MatchProfile=(Profile="Android_TegraK1",Match=((SourceType=SRC_GpuFamily,CompareType=CMP_Equal,MatchString="NVIDIA Tegra"),(SourceType=SRC_GlVersion,CompareType=CMP_Regex,MatchString="^OpenGL ES 3\\.")))
++MatchProfile=(Profile="Android_Tegra4",Match=((SourceType=SRC_GpuFamily,CompareType=CMP_Equal,MatchString="NVIDIA Tegra"),(SourceType=SRC_GlVersion,CompareType=CMP_Regex,MatchString="^OpenGL ES 2\\.")))
++MatchProfile=(Profile="Android_IntelHD_ES3",Match=((SourceType=SRC_GpuFamily,CompareType=CMP_Equal,MatchString="Intel(R) HD Graphics"),(SourceType=SRC_GlVersion,CompareType=CMP_Regex,MatchString="^OpenGL ES 3\\.")))
++MatchProfile=(Profile="Android_IntelHD",Match=((SourceType=SRC_GpuFamily,CompareType=CMP_Equal,MatchString="Intel(R) HD Graphics"),(SourceType=SRC_GlVersion,CompareType=CMP_Regex,MatchString="^OpenGL ES 2\\.")))
 +MatchProfile=(Profile="Android_Mali_4xx",Match=((SourceType=SRC_GpuFamily,CompareType=CMP_Regex,MatchString="^Mali\\-4")))
 +MatchProfile=(Profile="Android_Mali_T6xx",Match=((SourceType=SRC_GpuFamily,CompareType=CMP_Regex,MatchString="^Mali\\-T6")))
 +MatchProfile=(Profile="Android_Mali_T7xx",Match=((SourceType=SRC_GpuFamily,CompareType=CMP_Regex,MatchString="^Mali\\-T7")))
-+MatchProfile=(Profile="Android_Mid",Match=((SourceType=SRC_GlVersion,CompareType=CMP_Regex,MatchString="^OpenGL\\ ES 3\\.")))
-+MatchProfile=(Profile="Android_Low",Match=((SourceType=SRC_GlVersion,CompareType=CMP_Regex,MatchString="^OpenGL\\ ES 2\\.")))
++MatchProfile=(Profile="Android_Mid",Match=((SourceType=SRC_GlVersion,CompareType=CMP_Regex,MatchString="^OpenGL ES 3\\.")))
++MatchProfile=(Profile="Android_Low",Match=((SourceType=SRC_GlVersion,CompareType=CMP_Regex,MatchString="^OpenGL ES 2\\.")))
 
 [Android DeviceProfile]
 DeviceType=Android
@@ -230,23 +226,7 @@ BaseProfileName=Android_Mid
 ;This offset needs to be set for the mosaic fallback to work on Galaxy S4 (SAMSUNG-IGH-I337)
 ;+CVars=r.DemosaicVposOffset=0.5
 
-[Android_Adreno200 DeviceProfile]
-DeviceType=Android
-BaseProfileName=Android
-
-[Android_Adreno203 DeviceProfile]
-DeviceType=Android
-BaseProfileName=Android
-
-[Android_Adreno205 DeviceProfile]
-DeviceType=Android
-BaseProfileName=Android
-
-[Android_Adreno220 DeviceProfile]
-DeviceType=Android
-BaseProfileName=Android
-
-[Android_Adreno225 DeviceProfile]
+[Android_Adreno2xx DeviceProfile]
 DeviceType=Android
 BaseProfileName=Android
 
@@ -259,7 +239,7 @@ DeviceType=Android
 BaseProfileName=Android_Adreno330
 +CVars=r.DisjointTimerQueries=1
 
-[Android_Adreno420 DeviceProfile]
+[Android_Adreno4xx DeviceProfile]
 DeviceType=Android
 BaseProfileName=Android_High
 
