Fixed shader format name for [/Script/MacTargetPlatform.MacTargetSettings] in BaseEngine.ini

diff --git a/Engine/Config/BaseEngine.ini b/Engine/Config/BaseEngine.ini
--- a/Engine/Config/BaseEngine.ini
+++ b/Engine/Config/BaseEngine.ini
@@ -2353,7 +2353,7 @@ MinimumOSVersion=MSOS_Vista
 +TargetedRHIs=GLSL_150
  
 [/Script/MacTargetPlatform.MacTargetSettings]
-+TargetedRHIs=GLSL_150
++TargetedRHIs=GLSL_150_MAC
 
 [HMDPluginPriority]
 ; Since SteamVR also works with the Oculus Rift, give priority to the native Oculus plugin before trying SteamVR
