Merging using UE4-To-UE4-Releases-4.10
Android SDKAPILevel default change to matchndk

diff --git a/Engine/Config/BaseEngine.ini b/Engine/Config/BaseEngine.ini
--- a/Engine/Config/BaseEngine.ini
+++ b/Engine/Config/BaseEngine.ini
@@ -2330,7 +2330,7 @@ bMultiTargetFormat_ATC=true
 bMultiTargetFormat_ASTC=true
 
 [/Script/AndroidPlatformEditor.AndroidSDKSettings]
-SDKAPILevel=latest
+SDKAPILevel=matchndk
 NDKAPILevel=android-19
 
 
