Merging

fixed r.MSAA.CompositingSampleCount warning on startup

diff --git a/Engine/Config/BaseScalability.ini b/Engine/Config/BaseScalability.ini
--- a/Engine/Config/BaseScalability.ini
+++ b/Engine/Config/BaseScalability.ini
@@ -10,10 +10,6 @@
 ; Note: [ResolutionQuality] isn't using the usual 0/1/2/3 quality levels, the value directly maps to r.ScreenPercentage.
 ; This is because we want custom scaling aligned with the screen/window resolution with more fine grained control
 ;
-;-----------------------------------------------------------------------------------------------------------------
-;
-; Note: "r.MSAA.CompositingSampleCount" goes upto 8 however we're limiting it to 4 in a scale of 1/1/2/4
-; Meaning that it is off for the first two detail levels.
 
 ;-----------------------------------------------------------------------------------------------------------------
 
@@ -35,19 +31,15 @@ PerfIndexThresholds_EffectsQuality="Min 18 40 70"
 
 [AntiAliasingQuality@0]
 r.PostProcessAAQuality=0
-r.MSAA.CompositingSampleCount=1
 
 [AntiAliasingQuality@1]
 r.PostProcessAAQuality=2
-r.MSAA.CompositingSampleCount=1
 
 [AntiAliasingQuality@2]
 r.PostProcessAAQuality=3
-r.MSAA.CompositingSampleCount=2
 
 [AntiAliasingQuality@3]
 r.PostProcessAAQuality=4
-r.MSAA.CompositingSampleCount=4
 
 ;-----------------------------------------------------------------------------------------------------------------
 
