UE-21626

Some android target platform device profiles didnt exist in the base config. Added to rectify the parent and type issue.

May need further attention as the Android Device profiles arent all relevant.

diff --git a/Engine/Config/BaseDeviceProfiles.ini b/Engine/Config/BaseDeviceProfiles.ini
--- a/Engine/Config/BaseDeviceProfiles.ini
+++ b/Engine/Config/BaseDeviceProfiles.ini
@@ -361,3 +361,31 @@ DeviceType=LinuxServer
 BaseProfileName=Linux
 MeshLODSettings=
 TextureLODSettings=
+
+[Android_PVRTC DeviceProfile]
+DeviceType=Android
+BaseProfileName=Android
+
+[Android_ATC DeviceProfile]
+DeviceType=Android
+BaseProfileName=Android
+
+[Android_DXT DeviceProfile]
+DeviceType=Android
+BaseProfileName=Android
+
+[Android_ETC1 DeviceProfile]
+DeviceType=Android
+BaseProfileName=Android
+
+[Android_ETC2 DeviceProfile]
+DeviceType=Android
+BaseProfileName=Android
+
+[Android_ASTC DeviceProfile]
+DeviceType=Android
+BaseProfileName=Android
+
+[Android_Multi DeviceProfile]
+DeviceType=Android
+BaseProfileName=Android
\ No newline at end of file
