#UE4 Restored missing CookedTagsWhitelist list accidentally removed from previous merge

diff --git a/Engine/Config/BaseEngine.ini b/Engine/Config/BaseEngine.ini
--- a/Engine/Config/BaseEngine.ini
+++ b/Engine/Config/BaseEngine.ini
@@ -2106,6 +2106,16 @@ AspectRatioAxisConstraint=AspectRatio_MaintainXFOV
 +Class=ParticleSystem
 +Class=Texture2D
 
+[AssetRegistry]
+; Fill out this list with asset registry tags that are needed in cooked builds. All other tags are stripped during cook.
+; You can use a * wildcard to refer to either all classes or all tags for the class.
+; If your game needs all asset registry tags, simply add (Class=*,Tag=*). Be warned that this will include editor only data and it is advised you specify which tags you use.
+; ObjectLibraries
++CookedTagsWhitelist=(Class=Blueprint,Tag=ParentClass)
+; GameplayCueManager
++CookedTagsWhitelist=(Class=Blueprint,Tag=GeneratedClass)
++CookedTagsWhitelist=(Class=Blueprint,Tag=GameplayCueName)
+
 [AutomationTesting]
 ImportTestPath=../../Content/EditorAutomation/
 ImportTestPackagePath=/Engine/Content/EditorAutomation
