Merging CL 2746942 - Removed the hard-coded contentScaleFactor of 2 for iPhone 5S

diff --git a/Engine/Config/BaseDeviceProfiles.ini b/Engine/Config/BaseDeviceProfiles.ini
--- a/Engine/Config/BaseDeviceProfiles.ini
+++ b/Engine/Config/BaseDeviceProfiles.ini
@@ -152,7 +152,6 @@ BaseProfileName=IOS
 [iPhone5S DeviceProfile]
 DeviceType=IOS
 BaseProfileName=IOS
-+CVars=r.MobileContentScaleFactor=2
 +CVars=r.BloomQuality=1
 +CVars=r.DepthOfFieldQuality=1
 +CVars=r.LightShaftQuality=1
