Moved OrionGameplayDebuggingComponent references from BaseEngine.ini to DefaultEngine.ini #Orion

Recreated CL#2794030 from Orion/Dev-General

diff --git a/Engine/Config/BaseEngine.ini b/Engine/Config/BaseEngine.ini
--- a/Engine/Config/BaseEngine.ini
+++ b/Engine/Config/BaseEngine.ini
@@ -2299,8 +2299,8 @@ StartingBeaconPort=15000
 
 [/Script/GameplayDebugger.GameplayDebuggingReplicator]
 MaxEQSQueries=5
-DebugComponentClassName="/Script/OrionGameplayDebugger.OrionGameplayDebuggingComponent"
-DebugComponentHUDClassName="/Script/OrionGameplayDebugger.OrionGameplayDebuggingHUDComponent"
+DebugComponentClassName="/Script/GameplayDebugger.GameplayDebuggingComponent"
+DebugComponentHUDClassName="/Script/GameplayDebugger.GameplayDebuggingHUDComponent"
 DebugComponentControllerClassName = "/Script/GameplayDebugger.GameplayDebuggingControllerComponent"
 
 [/Script/GameplayDebugger.GameplayDebuggingHUDComponent]
