Adjusting the default value of FbxExportCompatibility in the config file from 3 to FBX_2013 so it will properly deserialize and fix the cook warnings.

diff --git a/Engine/Config/BaseEditorPerProjectUserSettings.ini b/Engine/Config/BaseEditorPerProjectUserSettings.ini
--- a/Engine/Config/BaseEditorPerProjectUserSettings.ini
+++ b/Engine/Config/BaseEditorPerProjectUserSettings.ini
@@ -55,7 +55,7 @@ bSCSEditorShowGrid=true
 
 ; Export
 bKeepAttachHierarchy=true
-FbxExportCompatibility=3
+FbxExportCompatibility=FBX_2013
 
 ; Misc
 HoverHighlightIntensity=0.000000
