Made VR-Basic Template, and set it to show up when you make a new level.

diff --git a/Engine/Config/BaseEngine.ini b/Engine/Config/BaseEngine.ini
--- a/Engine/Config/BaseEngine.ini
+++ b/Engine/Config/BaseEngine.ini
@@ -1634,6 +1634,7 @@ InEditorGameURLOptions=
 [/Script/UnrealEd.UnrealEdEngine]
 AutoSaveIndex=0
 +TemplateMapInfos=(ThumbnailTexture=Texture2D'/Engine/Maps/Templates/Thumbnails/Default.Default',Map="/Engine/Maps/Templates/Template_Default")
++TemplateMapInfos=(ThumbnailTexture=Texture2D'/Engine/Maps/Templates/Thumbnails/VR-Basic.VR-Basic',Map="/Engine/Maps/Templates/VR-Basic")
 
 +PackagesToBeFullyLoadedAtStartup=/Engine/EditorMaterials/Cross
 +PackagesToBeFullyLoadedAtStartup=/Engine/EditorMaterials/Cross_Mat
