VR Editor: Improved editor laser pointer ray tracing
- No longer pollutes game-specific .ini files with customized collison (ECC_EditorGizmo)
- Now supports interacting with objects that have collision disable on their component (but not static mesh asset)

diff --git a/Engine/Config/BaseEngine.ini b/Engine/Config/BaseEngine.ini
--- a/Engine/Config/BaseEngine.ini
+++ b/Engine/Config/BaseEngine.ini
@@ -2298,24 +2298,24 @@ bPromptForRemoteDebugOnEnsure=false
 EnableNiagara=false
 
 [/Script/Engine.CollisionProfile]
-+Profiles=(Name="NoCollision",CollisionEnabled=NoCollision,ObjectTypeName="WorldStatic",CustomResponses=((Channel="Visibility",Response=ECR_Ignore),(Channel="Camera",Response=ECR_Ignore),(Channel="EditorGizmo",Response=ECR_Ignore)),HelpMessage="No collision",bCanModify=False)
-+Profiles=(Name="BlockAll",CollisionEnabled=QueryAndPhysics,ObjectTypeName="WorldStatic",CustomResponses=((Channel="EditorGizmo",Response=ECR_Ignore)),HelpMessage="WorldStatic object that blocks all actors by default. All new custom channels will use its own default response. ",bCanModify=False)
-+Profiles=(Name="OverlapAll",CollisionEnabled=QueryOnly,ObjectTypeName="WorldStatic",CustomResponses=((Channel="WorldStatic",Response=ECR_Overlap),(Channel="Pawn",Response=ECR_Overlap),(Channel="Visibility",Response=ECR_Overlap),(Channel="WorldDynamic",Response=ECR_Overlap),(Channel="Camera",Response=ECR_Overlap),(Channel="PhysicsBody",Response=ECR_Overlap),(Channel="Vehicle",Response=ECR_Overlap),(Channel="Destructible",Response=ECR_Overlap),(Channel="EditorGizmo",Response=ECR_Ignore)),HelpMessage="WorldStatic object that overlaps all actors by default. All new custom channels will use its own default response. ",bCanModify=False)
-+Profiles=(Name="BlockAllDynamic",CollisionEnabled=QueryAndPhysics,ObjectTypeName="WorldDynamic",CustomResponses=((Channel="EditorGizmo",Response=ECR_Ignore)),HelpMessage="WorldDynamic object that blocks all actors by default. All new custom channels will use its own default response. ",bCanModify=False)
-+Profiles=(Name="OverlapAllDynamic",CollisionEnabled=QueryOnly,ObjectTypeName="WorldDynamic",CustomResponses=((Channel="WorldStatic",Response=ECR_Overlap),(Channel="Pawn",Response=ECR_Overlap),(Channel="Visibility",Response=ECR_Overlap),(Channel="WorldDynamic",Response=ECR_Overlap),(Channel="Camera",Response=ECR_Overlap),(Channel="PhysicsBody",Response=ECR_Overlap),(Channel="Vehicle",Response=ECR_Overlap),(Channel="Destructible",Response=ECR_Overlap),(Channel="EditorGizmo",Response=ECR_Ignore)),HelpMessage="WorldDynamic object that overlaps all actors by default. All new custom channels will use its own default response. ",bCanModify=False)
-+Profiles=(Name="IgnoreOnlyPawn",CollisionEnabled=QueryOnly,ObjectTypeName="WorldDynamic",CustomResponses=((Channel="Pawn",Response=ECR_Ignore),(Channel="Vehicle",Response=ECR_Ignore),(Channel="EditorGizmo",Response=ECR_Ignore)),HelpMessage="WorldDynamic object that ignores Pawn and Vehicle. All other channels will be set to default.",bCanModify=False)
-+Profiles=(Name="OverlapOnlyPawn",CollisionEnabled=QueryOnly,ObjectTypeName="WorldDynamic",CustomResponses=((Channel="Pawn",Response=ECR_Overlap),(Channel="Vehicle",Response=ECR_Overlap),(Channel="Camera",Response=ECR_Ignore),(Channel="EditorGizmo",Response=ECR_Ignore)),HelpMessage="WorldDynamic object that overlaps Pawn, Camera, and Vehicle. All other channels will be set to default. ",bCanModify=False)
-+Profiles=(Name="Pawn",CollisionEnabled=QueryAndPhysics,ObjectTypeName="Pawn",CustomResponses=((Channel="Visibility",Response=ECR_Ignore),(Channel="EditorGizmo",Response=ECR_Ignore)),HelpMessage="Pawn object. Can be used for capsule of any playerable character or AI. ",bCanModify=False)
-+Profiles=(Name="Spectator",CollisionEnabled=QueryOnly,ObjectTypeName="Pawn",CustomResponses=((Channel="WorldStatic",Response=ECR_Block),(Channel="Pawn",Response=ECR_Ignore),(Channel="Visibility",Response=ECR_Ignore),(Channel="WorldDynamic",Response=ECR_Ignore),(Channel="Camera",Response=ECR_Ignore),(Channel="PhysicsBody",Response=ECR_Ignore),(Channel="Vehicle",Response=ECR_Ignore),(Channel="Destructible",Response=ECR_Ignore),(Channel="EditorGizmo",Response=ECR_Ignore)),HelpMessage="Pawn object that ignores all other actors except WorldStatic.",bCanModify=False)
-+Profiles=(Name="CharacterMesh",CollisionEnabled=QueryOnly,ObjectTypeName="Pawn",CustomResponses=((Channel="Pawn",Response=ECR_Ignore),(Channel="Vehicle",Response=ECR_Ignore),(Channel="Visibility",Response=ECR_Ignore),(Channel="EditorGizmo",Response=ECR_Ignore)),HelpMessage="Pawn object that is used for Character Mesh. All other channels will be set to default.",bCanModify=False)
-+Profiles=(Name="PhysicsActor",CollisionEnabled=QueryAndPhysics,ObjectTypeName="PhysicsBody",CustomResponses=((Channel="EditorGizmo",Response=ECR_Ignore)),HelpMessage="Simulating actors",bCanModify=False)
-+Profiles=(Name="Destructible",CollisionEnabled=QueryAndPhysics,ObjectTypeName="Destructible",CustomResponses=((Channel="EditorGizmo",Response=ECR_Ignore)),HelpMessage="Destructible actors",bCanModify=False)
-+Profiles=(Name="InvisibleWall",CollisionEnabled=QueryAndPhysics,ObjectTypeName="WorldStatic",CustomResponses=((Channel="Visibility",Response=ECR_Ignore),(Channel="EditorGizmo",Response=ECR_Ignore)),HelpMessage="WorldStatic object that is invisible.",bCanModify=False)
-+Profiles=(Name="InvisibleWallDynamic",CollisionEnabled=QueryAndPhysics,ObjectTypeName="WorldDynamic",CustomResponses=((Channel="Visibility",Response=ECR_Ignore),(Channel="EditorGizmo",Response=ECR_Ignore)),HelpMessage="WorldDynamic object that is invisible.",bCanModify=False)
-+Profiles=(Name="Trigger",CollisionEnabled=QueryOnly,ObjectTypeName="WorldDynamic",CustomResponses=((Channel="WorldStatic",Response=ECR_Overlap),(Channel="Pawn",Response=ECR_Overlap),(Channel="Visibility",Response=ECR_Ignore),(Channel="WorldDynamic",Response=ECR_Overlap),(Channel="Camera",Response=ECR_Overlap),(Channel="PhysicsBody",Response=ECR_Overlap),(Channel="Vehicle",Response=ECR_Overlap),(Channel="Destructible",Response=ECR_Overlap),(Channel="EditorGizmo",Response=ECR_Ignore)),HelpMessage="WorldDynamic object that is used for trigger. All other channels will be set to default.",bCanModify=False)
-+Profiles=(Name="Ragdoll",CollisionEnabled=QueryAndPhysics,ObjectTypeName="PhysicsBody",CustomResponses=((Channel="Pawn",Response=ECR_Ignore),(Channel="EditorGizmo",Response=ECR_Ignore)),HelpMessage="Simulating Skeletal Mesh Component. All other channels will be set to default.",bCanModify=False)
-+Profiles=(Name="Vehicle",CollisionEnabled=QueryAndPhysics,ObjectTypeName="Vehicle",CustomResponses=((Channel="EditorGizmo",Response=ECR_Ignore)),HelpMessage="Vehicle object that blocks Vehicle, WorldStatic, and WorldDynamic. All other channels will be set to default.",bCanModify=False)
-+Profiles=(Name="UI",CollisionEnabled=QueryOnly,ObjectTypeName="WorldDynamic",CustomResponses=((Channel="WorldStatic",Response=ECR_Overlap),(Channel="Pawn",Response=ECR_Overlap),(Channel="Visibility",Response=ECR_Block),(Channel="WorldDynamic",Response=ECR_Overlap),(Channel="Camera",Response=ECR_Overlap),(Channel="PhysicsBody",Response=ECR_Overlap),(Channel="Vehicle",Response=ECR_Overlap),(Channel="Destructible",Response=ECR_Overlap),(Channel="EditorGizmo",Response=ECR_Ignore)),HelpMessage="WorldStatic object that overlaps all actors by default. All new custom channels will use its own default response. ",bCanModify=False)
++Profiles=(Name="NoCollision",CollisionEnabled=NoCollision,ObjectTypeName="WorldStatic",CustomResponses=((Channel="Visibility",Response=ECR_Ignore),(Channel="Camera",Response=ECR_Ignore)),HelpMessage="No collision",bCanModify=False)
++Profiles=(Name="BlockAll",CollisionEnabled=QueryAndPhysics,ObjectTypeName="WorldStatic",CustomResponses=,HelpMessage="WorldStatic object that blocks all actors by default. All new custom channels will use its own default response. ",bCanModify=False)
++Profiles=(Name="OverlapAll",CollisionEnabled=QueryOnly,ObjectTypeName="WorldStatic",CustomResponses=((Channel="WorldStatic",Response=ECR_Overlap),(Channel="Pawn",Response=ECR_Overlap),(Channel="Visibility",Response=ECR_Overlap),(Channel="WorldDynamic",Response=ECR_Overlap),(Channel="Camera",Response=ECR_Overlap),(Channel="PhysicsBody",Response=ECR_Overlap),(Channel="Vehicle",Response=ECR_Overlap),(Channel="Destructible",Response=ECR_Overlap)),HelpMessage="WorldStatic object that overlaps all actors by default. All new custom channels will use its own default response. ",bCanModify=False)
++Profiles=(Name="BlockAllDynamic",CollisionEnabled=QueryAndPhysics,ObjectTypeName="WorldDynamic",CustomResponses=,HelpMessage="WorldDynamic object that blocks all actors by default. All new custom channels will use its own default response. ",bCanModify=False)
++Profiles=(Name="OverlapAllDynamic",CollisionEnabled=QueryOnly,ObjectTypeName="WorldDynamic",CustomResponses=((Channel="WorldStatic",Response=ECR_Overlap),(Channel="Pawn",Response=ECR_Overlap),(Channel="Visibility",Response=ECR_Overlap),(Channel="WorldDynamic",Response=ECR_Overlap),(Channel="Camera",Response=ECR_Overlap),(Channel="PhysicsBody",Response=ECR_Overlap),(Channel="Vehicle",Response=ECR_Overlap),(Channel="Destructible",Response=ECR_Overlap)),HelpMessage="WorldDynamic object that overlaps all actors by default. All new custom channels will use its own default response. ",bCanModify=False)
++Profiles=(Name="IgnoreOnlyPawn",CollisionEnabled=QueryOnly,ObjectTypeName="WorldDynamic",CustomResponses=((Channel="Pawn",Response=ECR_Ignore),(Channel="Vehicle",Response=ECR_Ignore)),HelpMessage="WorldDynamic object that ignores Pawn and Vehicle. All other channels will be set to default.",bCanModify=False)
++Profiles=(Name="OverlapOnlyPawn",CollisionEnabled=QueryOnly,ObjectTypeName="WorldDynamic",CustomResponses=((Channel="Pawn",Response=ECR_Overlap),(Channel="Vehicle",Response=ECR_Overlap),(Channel="Camera",Response=ECR_Ignore)),HelpMessage="WorldDynamic object that overlaps Pawn, Camera, and Vehicle. All other channels will be set to default. ",bCanModify=False)
++Profiles=(Name="Pawn",CollisionEnabled=QueryAndPhysics,ObjectTypeName="Pawn",CustomResponses=((Channel="Visibility",Response=ECR_Ignore)),HelpMessage="Pawn object. Can be used for capsule of any playerable character or AI. ",bCanModify=False)
++Profiles=(Name="Spectator",CollisionEnabled=QueryOnly,ObjectTypeName="Pawn",CustomResponses=((Channel="WorldStatic",Response=ECR_Block),(Channel="Pawn",Response=ECR_Ignore),(Channel="Visibility",Response=ECR_Ignore),(Channel="WorldDynamic",Response=ECR_Ignore),(Channel="Camera",Response=ECR_Ignore),(Channel="PhysicsBody",Response=ECR_Ignore),(Channel="Vehicle",Response=ECR_Ignore),(Channel="Destructible",Response=ECR_Ignore)),HelpMessage="Pawn object that ignores all other actors except WorldStatic.",bCanModify=False)
++Profiles=(Name="CharacterMesh",CollisionEnabled=QueryOnly,ObjectTypeName="Pawn",CustomResponses=((Channel="Pawn",Response=ECR_Ignore),(Channel="Vehicle",Response=ECR_Ignore),(Channel="Visibility",Response=ECR_Ignore)),HelpMessage="Pawn object that is used for Character Mesh. All other channels will be set to default.",bCanModify=False)
++Profiles=(Name="PhysicsActor",CollisionEnabled=QueryAndPhysics,ObjectTypeName="PhysicsBody",CustomResponses=,HelpMessage="Simulating actors",bCanModify=False)
++Profiles=(Name="Destructible",CollisionEnabled=QueryAndPhysics,ObjectTypeName="Destructible",CustomResponses=,HelpMessage="Destructible actors",bCanModify=False)
++Profiles=(Name="InvisibleWall",CollisionEnabled=QueryAndPhysics,ObjectTypeName="WorldStatic",CustomResponses=((Channel="Visibility",Response=ECR_Ignore)),HelpMessage="WorldStatic object that is invisible.",bCanModify=False)
++Profiles=(Name="InvisibleWallDynamic",CollisionEnabled=QueryAndPhysics,ObjectTypeName="WorldDynamic",CustomResponses=((Channel="Visibility",Response=ECR_Ignore)),HelpMessage="WorldDynamic object that is invisible.",bCanModify=False)
++Profiles=(Name="Trigger",CollisionEnabled=QueryOnly,ObjectTypeName="WorldDynamic",CustomResponses=((Channel="WorldStatic",Response=ECR_Overlap),(Channel="Pawn",Response=ECR_Overlap),(Channel="Visibility",Response=ECR_Ignore),(Channel="WorldDynamic",Response=ECR_Overlap),(Channel="Camera",Response=ECR_Overlap),(Channel="PhysicsBody",Response=ECR_Overlap),(Channel="Vehicle",Response=ECR_Overlap),(Channel="Destructible",Response=ECR_Overlap)),HelpMessage="WorldDynamic object that is used for trigger. All other channels will be set to default.",bCanModify=False)
++Profiles=(Name="Ragdoll",CollisionEnabled=QueryAndPhysics,ObjectTypeName="PhysicsBody",CustomResponses=((Channel="Pawn",Response=ECR_Ignore)),HelpMessage="Simulating Skeletal Mesh Component. All other channels will be set to default.",bCanModify=False)
++Profiles=(Name="Vehicle",CollisionEnabled=QueryAndPhysics,ObjectTypeName="Vehicle",CustomResponses=,HelpMessage="Vehicle object that blocks Vehicle, WorldStatic, and WorldDynamic. All other channels will be set to default.",bCanModify=False)
++Profiles=(Name="UI",CollisionEnabled=QueryOnly,ObjectTypeName="WorldDynamic",CustomResponses=((Channel="WorldStatic",Response=ECR_Overlap),(Channel="Pawn",Response=ECR_Overlap),(Channel="Visibility",Response=ECR_Block),(Channel="WorldDynamic",Response=ECR_Overlap),(Channel="Camera",Response=ECR_Overlap),(Channel="PhysicsBody",Response=ECR_Overlap),(Channel="Vehicle",Response=ECR_Overlap),(Channel="Destructible",Response=ECR_Overlap)),HelpMessage="WorldStatic object that overlaps all actors by default. All new custom channels will use its own default response. ",bCanModify=False)
 
 +OldProfiles=(Name="BlockingVolume",CollisionEnabled=QueryAndPhysics,ObjectTypeName=WorldStatic,CustomResponses=((Channel=Visibility, Response=ECR_Ignore)))
 ; ECC_Mover is gone, but it doens't matter. Only thing that matters is the name
diff --git a/Engine/Source/Editor/VREditor/Gizmo/VREditorGizmoHandle.cpp b/Engine/Source/Editor/VREditor/Gizmo/VREditorGizmoHandle.cpp
--- a/Engine/Source/Editor/VREditor/Gizmo/VREditorGizmoHandle.cpp
+++ b/Engine/Source/Editor/VREditor/Gizmo/VREditorGizmoHandle.cpp
@@ -221,6 +221,7 @@ class UStaticMeshComponent* UVREditorGizmoHandleGroup::CreateMeshHandle( class U
 	HandleComponent->SetCollisionEnabled( ECollisionEnabled::QueryOnly );
 	HandleComponent->SetCollisionResponseToAllChannels( ECR_Ignore );
 	HandleComponent->SetCollisionResponseToChannel( ECC_EditorGizmo, ECollisionResponse::ECR_Block );
+	HandleComponent->SetCollisionObjectType( ECC_EditorGizmo );
 
 	HandleComponent->bGenerateOverlapEvents = false;
 	HandleComponent->SetCanEverAffectNavigation( false );
diff --git a/Engine/Source/Editor/VREditor/Gizmo/VREditorUniformScaleGizmoHandle.cpp b/Engine/Source/Editor/VREditor/Gizmo/VREditorUniformScaleGizmoHandle.cpp
--- a/Engine/Source/Editor/VREditor/Gizmo/VREditorUniformScaleGizmoHandle.cpp
+++ b/Engine/Source/Editor/VREditor/Gizmo/VREditorUniformScaleGizmoHandle.cpp
@@ -29,6 +29,7 @@ UVREditorUniformScaleGizmoHandleGroup::UVREditorUniformScaleGizmoHandleGroup()
 	UniformScaleHandle->SetCollisionEnabled( ECollisionEnabled::QueryOnly );
 	UniformScaleHandle->SetCollisionResponseToAllChannels( ECR_Ignore );
 	UniformScaleHandle->SetCollisionResponseToChannel( ECC_EditorGizmo, ECollisionResponse::ECR_Block );
+	UniformScaleHandle->SetCollisionObjectType( ECC_EditorGizmo );
 
 	UniformScaleHandle->bGenerateOverlapEvents = false;
 	UniformScaleHandle->SetCanEverAffectNavigation( false );
diff --git a/Engine/Source/Editor/VREditor/VREditorDockableWindow.cpp b/Engine/Source/Editor/VREditor/VREditorDockableWindow.cpp
--- a/Engine/Source/Editor/VREditor/VREditorDockableWindow.cpp
+++ b/Engine/Source/Editor/VREditor/VREditorDockableWindow.cpp
@@ -36,8 +36,9 @@ AVREditorDockableWindow::AVREditorDockableWindow()
 		WindowMeshComponent->AttachTo( RootComponent );
 
 		WindowMeshComponent->SetCollisionEnabled( ECollisionEnabled::QueryOnly );
-		WindowMeshComponent->SetCollisionResponseToAllChannels( ECollisionResponse::ECR_Block );
+		WindowMeshComponent->SetCollisionResponseToAllChannels( ECollisionResponse::ECR_Ignore );
 		WindowMeshComponent->SetCollisionResponseToChannel( ECC_EditorGizmo, ECollisionResponse::ECR_Block );
+		WindowMeshComponent->SetCollisionObjectType( ECC_EditorGizmo );
 
 		WindowMeshComponent->bGenerateOverlapEvents = false;
 		WindowMeshComponent->SetCanEverAffectNavigation( false );
@@ -60,8 +61,9 @@ AVREditorDockableWindow::AVREditorDockableWindow()
 		SelectionBarMeshComponent->AttachTo( RootComponent );
 
 		SelectionBarMeshComponent->SetCollisionEnabled( ECollisionEnabled::QueryOnly );
-		SelectionBarMeshComponent->SetCollisionResponseToAllChannels( ECollisionResponse::ECR_Block );
+		SelectionBarMeshComponent->SetCollisionResponseToAllChannels( ECollisionResponse::ECR_Ignore );
 		SelectionBarMeshComponent->SetCollisionResponseToChannel( ECC_EditorGizmo, ECollisionResponse::ECR_Block );
+		SelectionBarMeshComponent->SetCollisionObjectType( ECC_EditorGizmo );
 
 		SelectionBarMeshComponent->bGenerateOverlapEvents = false;
 		SelectionBarMeshComponent->SetCanEverAffectNavigation( false );
@@ -104,8 +106,9 @@ AVREditorDockableWindow::AVREditorDockableWindow()
 		CloseButtonMeshComponent->AttachTo( RootComponent );
 
 		CloseButtonMeshComponent->SetCollisionEnabled( ECollisionEnabled::QueryOnly );
-		CloseButtonMeshComponent->SetCollisionResponseToAllChannels( ECollisionResponse::ECR_Block );
+		CloseButtonMeshComponent->SetCollisionResponseToAllChannels( ECollisionResponse::ECR_Ignore );
 		CloseButtonMeshComponent->SetCollisionResponseToChannel( ECC_EditorGizmo, ECollisionResponse::ECR_Block );
+		CloseButtonMeshComponent->SetCollisionObjectType( ECC_EditorGizmo );
 
 		CloseButtonMeshComponent->bGenerateOverlapEvents = false;
 		CloseButtonMeshComponent->SetCanEverAffectNavigation( false );
diff --git a/Engine/Source/Editor/VREditor/VREditorMode.cpp b/Engine/Source/Editor/VREditor/VREditorMode.cpp
--- a/Engine/Source/Editor/VREditor/VREditorMode.cpp
+++ b/Engine/Source/Editor/VREditor/VREditorMode.cpp
@@ -1811,8 +1811,19 @@ FHitResult FVREditorMode::GetHitResultFromLaserPointer( int32 HandIndex, TArray<
 				TraceParams.AddIgnoredActors( *OptionalListOfIgnoredActors );
 			}
 
+			bool bHit = false;
 			FHitResult HitResult;
-			if( GetWorld()->LineTraceSingleByChannel( HitResult, LaserPointerStart, LaserPointerEnd, CollisionChannel, TraceParams, ResponseParam ) )
+			if( bOnlyEditorGizmos )
+			{
+				bHit = GetWorld()->LineTraceSingleByChannel( HitResult, LaserPointerStart, LaserPointerEnd, CollisionChannel, TraceParams, ResponseParam );
+			}
+			else
+			{
+				FCollisionObjectQueryParams EverythingButGizmos( FCollisionObjectQueryParams::AllObjects );
+				EverythingButGizmos.RemoveObjectTypesToQuery( ECC_EditorGizmo );
+				bHit = GetWorld()->LineTraceSingleByObjectType( HitResult, LaserPointerStart, LaserPointerEnd, EverythingButGizmos, TraceParams );
+			}
+			if( bHit )
 			{
 				// Is this better than what we have already?  Always prefer transform gizmos even if they were not using
 				// ECC_EditorGizmo (some gizmo handles are opaque and use ECC_Visibility as their collision channel)
diff --git a/Engine/Source/Runtime/Engine/Classes/Engine/EngineTypes.h b/Engine/Source/Runtime/Engine/Classes/Engine/EngineTypes.h
--- a/Engine/Source/Runtime/Engine/Classes/Engine/EngineTypes.h
+++ b/Engine/Source/Runtime/Engine/Classes/Engine/EngineTypes.h
@@ -403,7 +403,7 @@ enum ECollisionChannel
 	ECC_PhysicsBody UMETA(DisplayName="PhysicsBody"),
 	ECC_Vehicle UMETA(DisplayName="Vehicle"),
 	ECC_Destructible UMETA(DisplayName="Destructible"),
-	ECC_EditorGizmo UMETA(DisplayName="EditorGizmo", TraceQuery="1"),
+	ECC_EditorGizmo UMETA(Hidden, DisplayName="EditorGizmo", TraceQuery="1"),
 	// @NOTE : when you add more here for predefined engine channel
 	// please change the max in the CollisionProfile
 	// search ECC_EditorGizmo
@@ -414,7 +414,6 @@ enum ECollisionChannel
 	ECC_EngineTraceChannel3 UMETA(Hidden),
 	ECC_EngineTraceChannel4 UMETA(Hidden), 
 	ECC_EngineTraceChannel5 UMETA(Hidden),
-	ECC_EngineTraceChannel6 UMETA(Hidden),
 
 	// in order to use this custom channels
 	// we recommend to define in your local file
@@ -691,6 +690,8 @@ struct ENGINE_API FCollisionResponseContainer
 	UPROPERTY(EditAnywhere, BlueprintReadOnly, Category=CollisionResponseContainer, meta=(DisplayName="Destructible"))
 	TEnumAsByte<enum ECollisionResponse> Destructible;    // 7
 
+	UPROPERTY(EditAnywhere, BlueprintReadOnly, Category=CollisionResponseContainer, meta=(DisplayName="EditorGizmo"))
+	TEnumAsByte<enum ECollisionResponse> EditorGizmo;    // 7
 
 	///////////////////////////////////////
 	// Unspecified Engine Trace Channels
@@ -710,9 +711,6 @@ struct ENGINE_API FCollisionResponseContainer
 	UPROPERTY(EditAnywhere, BlueprintReadOnly, Category=CollisionResponseContainer)
 	TEnumAsByte<enum ECollisionResponse> EngineTraceChannel5;    // 12
 
-	UPROPERTY(EditAnywhere, BlueprintReadOnly, Category=CollisionResponseContainer)
-	TEnumAsByte<enum ECollisionResponse> EngineTraceChannel6;    // 13
-
 	///////////////////////////////////////
 	// in order to use this custom channels
 	// we recommend to define in your local file
@@ -797,14 +795,14 @@ struct ENGINE_API FCollisionResponseContainer
 			uint8 PhysicsBody;			// 5
 			uint8 Vehicle;				// 6
 			uint8 Destructible;			// 7
+			uint8 EditorGizmo;			// 8
 
 			// Unspecified Engine Trace Channels
-			uint8 EngineTraceChannel1;   // 8
-			uint8 EngineTraceChannel2;   // 9
-			uint8 EngineTraceChannel3;   // 10
-			uint8 EngineTraceChannel4;   // 11
-			uint8 EngineTraceChannel5;   // 12
-			uint8 EngineTraceChannel6;   // 13
+			uint8 EngineTraceChannel1;   // 9
+			uint8 EngineTraceChannel2;   // 10
+			uint8 EngineTraceChannel3;   // 11
+			uint8 EngineTraceChannel4;   // 12
+			uint8 EngineTraceChannel5;   // 13
 
 			// Unspecified Game Trace Channels
 			uint8 GameTraceChannel1;     // 14
diff --git a/Engine/Source/Runtime/Engine/Private/Collision/CollisionProfile.cpp b/Engine/Source/Runtime/Engine/Private/Collision/CollisionProfile.cpp
--- a/Engine/Source/Runtime/Engine/Private/Collision/CollisionProfile.cpp
+++ b/Engine/Source/Runtime/Engine/Private/Collision/CollisionProfile.cpp
@@ -356,6 +356,7 @@ void UCollisionProfile::LoadProfileConfig(bool bForceInit)
 	// Now Load Channel setups, and set display names if customized
 	// also initialize DefaultResposneContainer with default response for each channel
 	FCollisionResponseContainer::DefaultResponseContainer.SetAllChannels(ECR_Block);
+	FCollisionResponseContainer::DefaultResponseContainer.SetResponse( ECC_EditorGizmo, ECR_Ignore );
 
 	for (auto Iter=DefaultChannelResponses.CreateConstIterator(); Iter; ++Iter)
 	{
