fix for cvar warnings in Windows DDC

diff --git a/Engine/Config/BaseDeviceProfiles.ini b/Engine/Config/BaseDeviceProfiles.ini
--- a/Engine/Config/BaseDeviceProfiles.ini
+++ b/Engine/Config/BaseDeviceProfiles.ini
@@ -82,8 +82,8 @@ TextureLODGroups=(Group=TEXTUREGROUP_World,MinLODSize=1,MaxLODSize=4096,LODBias=
 [Windows DeviceProfile]
 DeviceType=Windows
 BaseProfileName=
-+CVars=r.Vulkan.UseGLSL=0
-+CVars=r.Vulkan.EnableValidation=0
+// +CVars=r.Vulkan.UseGLSL=0
+// +CVars=r.Vulkan.EnableValidation=0
 
 [WindowsNoEditor DeviceProfile]
 DeviceType=WindowsNoEditor
@@ -266,8 +266,8 @@ BaseProfileName=
 +CVars=r.LightShaftQuality=0
 +CVars=r.RefractionQuality=0
 +CVars=r.ShadowQuality=2
-+CVars=r.Vulkan.UseGLSL=2
-+CVars=r.Vulkan.EnableValidation=0
+// +CVars=r.Vulkan.UseGLSL=2
+// +CVars=r.Vulkan.EnableValidation=0
 
 
 [Android_Low DeviceProfile]
