Updating recommended AMD and Nvidia drivers - RB Marcus Wassmer
AMD - 16.40.3111.0
NV - 375.70

diff --git a/Engine/Config/BaseHardware.ini b/Engine/Config/BaseHardware.ini
--- a/Engine/Config/BaseHardware.ini
+++ b/Engine/Config/BaseHardware.ini
@@ -32,7 +32,7 @@
 
 [GPU_NVIDIA]
 ; a newer one is out but has issues (see below)
-SuggestedDriverVersion="372.90"
+SuggestedDriverVersion="375.70"
 
 ; Seems this driver was not setting uninitialized global variables
 ; (static, non const) to 0 but leaving them uninitialized.
@@ -51,7 +51,7 @@ SuggestedDriverVersion="372.90"
 
 [GPU_AMD]
 // From OR-27051 as of 8/4/2016 is 16.300.2311.0 7/18/2016 (Crimson Edition 16.9.2)
-SuggestedDriverVersion="16.300.2311.0"
+SuggestedDriverVersion="16.400.3111.0"
 
 // From OR-27051
 +Blacklist=(DriverVersion="==15.300.1025.1001", Reason="Crashes with Paragon content using more recent GPU features")
