- Removed confusung and invalid ini settings

diff --git a/Engine/Config/BaseDeviceProfiles.ini b/Engine/Config/BaseDeviceProfiles.ini
--- a/Engine/Config/BaseDeviceProfiles.ini
+++ b/Engine/Config/BaseDeviceProfiles.ini
@@ -63,9 +63,6 @@
 +DeviceProfileNameAndTypes=Android_Mali_T7xx,Android
 +DeviceProfileNameAndTypes=Android_Mali_T8xx,Android
 +DeviceProfileNameAndTypes=Android_Mali_T8xx_No_Vulkan,Android
-+DeviceProfileNameAndTypes=Switch,Switch
-+DeviceProfileNameAndTypes=SwitchHandheld,Switch
-+DeviceProfileNameAndTypes=SwitchConsole,Switch
 
 [/Script/Engine.TextureLODSettings]
 @TextureLODGroups=Group
@@ -95,31 +92,6 @@ TextureLODGroups=(Group=TEXTUREGROUP_World,MinLODSize=1,MaxLODSize=4096,LODBias=
 +TextureLODGroups=(Group=TEXTUREGROUP_Bokeh,MinLODSize=1,MaxLODSize=256,LODBias=0,MinMagFilter=linear,MipFilter=linear,MipGenSettings=TMGS_SimpleAverage)
 +TextureLODGroups=(Group=TEXTUREGROUP_Pixels2D,MinLODSize=1,MaxLODSize=4096,LODBias=0,MinMagFilter=point,MipFilter=point,MipGenSettings=TMGS_SimpleAverage)
 
-[Switch DeviceProfile]
-+CVars=sg.ViewDistanceQuality=2
-+CVars=sg.AntiAliasingQuality=2
-+CVars=sg.ShadowQuality=2
-+CVars=sg.PostProcessQuality=2
-+CVars=sg.TextureQuality=2
-+CVars=sg.EffectsQuality=2
-+CVars=r.ScreenPercentage=100
-
-[SwitchConsole DeviceProfile]
-DeviceType=Switch
-BaseProfileName=Switch
-
-
-[SwitchHandheld DeviceProfile]
-DeviceType=Switch
-BaseProfileName=Switch
-+CVars=sg.ViewDistanceQuality=1
-+CVars=sg.AntiAliasingQuality=1
-+CVars=sg.ShadowQuality=1
-+CVars=sg.PostProcessQuality=1
-+CVars=sg.TextureQuality=1
-+CVars=sg.EffectsQuality=1
-+CVars=r.ScreenPercentage=66
-
 [Windows DeviceProfile]
 DeviceType=Windows
 BaseProfileName=
