Reenable audio threading

diff --git a/Engine/Config/BaseEngine.ini b/Engine/Config/BaseEngine.ini
--- a/Engine/Config/BaseEngine.ini
+++ b/Engine/Config/BaseEngine.ini
@@ -1688,7 +1688,7 @@ gc.MergeGCClusters=False
 MaxChannels=32
 CommonAudioPoolSize=0
 UnfocusedVolumeMultiplier=0.0
-UseAudioThread=false
+UseAudioThread=true
 EnableAudioMixer=false
 
 [AudioChannelAzimuthMap]
