Fix shootergame to cook maps properly, was broken by merge from framework

diff --git a/Engine/Config/BaseGame.ini b/Engine/Config/BaseGame.ini
--- a/Engine/Config/BaseGame.ini
+++ b/Engine/Config/BaseGame.ini
@@ -122,4 +122,4 @@ QueryCountWarningInterval=30.0
 
 [/Script/Engine.AssetManagerSettings]
 +PrimaryAssetTypesToScan=(PrimaryAssetType="Map",AssetBaseClass=/Script/Engine.World,bHasBlueprintClasses=False,bIsEditorOnly=True,Directories=((Path="/Game/Maps")))
-+PrimaryAssetTypesToScan=(PrimaryAssetType="PrimaryAssetLabel",AssetBaseClass=/Script/Engine.PrimaryAssetLabel,bHasBlueprintClasses=False,bIsEditorOnly=True,Directories=((Path="/Game")),Rules=(CookRule=NeverCook))
++PrimaryAssetTypesToScan=(PrimaryAssetType="PrimaryAssetLabel",AssetBaseClass=/Script/Engine.PrimaryAssetLabel,bHasBlueprintClasses=False,bIsEditorOnly=True,Directories=((Path="/Game")))
