HTML5 - force max texture sizes to 256x256

zen garden tests has shown asset data file go from 1.71 GB down to 544 MB
this helps prevents chrome from doing: "aww snap!" anymore....

#rn

diff --git a/Engine/Config/BaseDeviceProfiles.ini b/Engine/Config/BaseDeviceProfiles.ini
--- a/Engine/Config/BaseDeviceProfiles.ini
+++ b/Engine/Config/BaseDeviceProfiles.ini
@@ -447,6 +447,36 @@ BaseProfileName=Android_Mali_T8xx
 DeviceType=HTML5
 BaseProfileName=
 +CVars=r.RefractionQuality=1
++TextureLODGroups=(Group=TEXTUREGROUP_World,LODBias=0,NumStreamedMips=-1,MipGenSettings=TMGS_SimpleAverage,MinLODSize=1,MaxLODSize=256,MinMagFilter="aniso",MipFilter="point")
++TextureLODGroups=(Group=TEXTUREGROUP_WorldNormalMap,LODBias=0,NumStreamedMips=-1,MipGenSettings=TMGS_SimpleAverage,MinLODSize=1,MaxLODSize=256,MinMagFilter="aniso",MipFilter="point")
++TextureLODGroups=(Group=TEXTUREGROUP_WorldSpecular,LODBias=0,NumStreamedMips=-1,MipGenSettings=TMGS_SimpleAverage,MinLODSize=1,MaxLODSize=256,MinMagFilter="aniso",MipFilter="point")
++TextureLODGroups=(Group=TEXTUREGROUP_Character,LODBias=0,NumStreamedMips=-1,MipGenSettings=TMGS_SimpleAverage,MinLODSize=1,MaxLODSize=256,MinMagFilter="aniso",MipFilter="point")
++TextureLODGroups=(Group=TEXTUREGROUP_CharacterNormalMap,LODBias=0,NumStreamedMips=-1,MipGenSettings=TMGS_SimpleAverage,MinLODSize=1,MaxLODSize=256,MinMagFilter="aniso",MipFilter="point")
++TextureLODGroups=(Group=TEXTUREGROUP_CharacterSpecular,LODBias=0,NumStreamedMips=-1,MipGenSettings=TMGS_SimpleAverage,MinLODSize=1,MaxLODSize=256,MinMagFilter="aniso",MipFilter="point")
++TextureLODGroups=(Group=TEXTUREGROUP_Weapon,LODBias=0,NumStreamedMips=-1,MipGenSettings=TMGS_SimpleAverage,MinLODSize=1,MaxLODSize=256,MinMagFilter="aniso",MipFilter="point")
++TextureLODGroups=(Group=TEXTUREGROUP_WeaponNormalMap,LODBias=0,NumStreamedMips=-1,MipGenSettings=TMGS_SimpleAverage,MinLODSize=1,MaxLODSize=256,MinMagFilter="aniso",MipFilter="point")
++TextureLODGroups=(Group=TEXTUREGROUP_WeaponSpecular,LODBias=0,NumStreamedMips=-1,MipGenSettings=TMGS_SimpleAverage,MinLODSize=1,MaxLODSize=256,MinMagFilter="aniso",MipFilter="point")
++TextureLODGroups=(Group=TEXTUREGROUP_Vehicle,LODBias=0,NumStreamedMips=-1,MipGenSettings=TMGS_SimpleAverage,MinLODSize=1,MaxLODSize=256,MinMagFilter="aniso",MipFilter="point")
++TextureLODGroups=(Group=TEXTUREGROUP_VehicleNormalMap,LODBias=0,NumStreamedMips=-1,MipGenSettings=TMGS_SimpleAverage,MinLODSize=1,MaxLODSize=256,MinMagFilter="aniso",MipFilter="point")
++TextureLODGroups=(Group=TEXTUREGROUP_VehicleSpecular,LODBias=0,NumStreamedMips=-1,MipGenSettings=TMGS_SimpleAverage,MinLODSize=1,MaxLODSize=256,MinMagFilter="aniso",MipFilter="point")
++TextureLODGroups=(Group=TEXTUREGROUP_Cinematic,LODBias=0,NumStreamedMips=-1,MipGenSettings=TMGS_SimpleAverage,MinLODSize=1,MaxLODSize=256,MinMagFilter="aniso",MipFilter="point")
++TextureLODGroups=(Group=TEXTUREGROUP_Effects,LODBias=0,NumStreamedMips=-1,MipGenSettings=TMGS_SimpleAverage,MinLODSize=1,MaxLODSize=256,MinMagFilter="linear",MipFilter="point")
++TextureLODGroups=(Group=TEXTUREGROUP_EffectsNotFiltered,LODBias=0,NumStreamedMips=-1,MipGenSettings=TMGS_SimpleAverage,MinLODSize=1,MaxLODSize=256,MinMagFilter="aniso",MipFilter="point")
++TextureLODGroups=(Group=TEXTUREGROUP_Skybox,LODBias=0,NumStreamedMips=-1,MipGenSettings=TMGS_SimpleAverage,MinLODSize=1,MaxLODSize=256,MinMagFilter="aniso",MipFilter="point")
++TextureLODGroups=(Group=TEXTUREGROUP_UI,LODBias=0,NumStreamedMips=-1,MipGenSettings=TMGS_SimpleAverage,MinLODSize=1,MaxLODSize=256,MinMagFilter="aniso",MipFilter="point")
++TextureLODGroups=(Group=TEXTUREGROUP_Lightmap,LODBias=0,NumStreamedMips=-1,MipGenSettings=TMGS_SimpleAverage,MinLODSize=1,MaxLODSize=256,MinMagFilter="aniso",MipFilter="point")
++TextureLODGroups=(Group=TEXTUREGROUP_RenderTarget,LODBias=0,NumStreamedMips=-1,MipGenSettings=TMGS_SimpleAverage,MinLODSize=1,MaxLODSize=256,MinMagFilter="aniso",MipFilter="point")
++TextureLODGroups=(Group=TEXTUREGROUP_MobileFlattened,LODBias=0,NumStreamedMips=-1,MipGenSettings=TMGS_SimpleAverage,MinLODSize=1,MaxLODSize=256,MinMagFilter="aniso",MipFilter="point")
++TextureLODGroups=(Group=TEXTUREGROUP_ProcBuilding_Face,LODBias=0,NumStreamedMips=-1,MipGenSettings=TMGS_SimpleAverage,MinLODSize=1,MaxLODSize=256,MinMagFilter="Aniso",MipFilter="Point")
++TextureLODGroups=(Group=TEXTUREGROUP_ProcBuilding_LightMap,LODBias=0,NumStreamedMips=-1,MipGenSettings=TMGS_SimpleAverage,MinLODSize=1,MaxLODSize=256,MinMagFilter="Aniso",MipFilter="Point")
++TextureLODGroups=(Group=TEXTUREGROUP_Shadowmap,LODBias=0,NumStreamedMips=3,MipGenSettings=TMGS_SimpleAverage,MinLODSize=1,MaxLODSize=256,MinMagFilter="aniso",MipFilter="point")
++TextureLODGroups=(Group=TEXTUREGROUP_ColorLookupTable,LODBias=0,NumStreamedMips=-1,MipGenSettings=TMGS_SimpleAverage,MinLODSize=1,MaxLODSize=256,MinMagFilter="Aniso",MipFilter="Point")
++TextureLODGroups=(Group=TEXTUREGROUP_Terrain_Heightmap,LODBias=0,NumStreamedMips=-1,MipGenSettings=TMGS_SimpleAverage,MinLODSize=1,MaxLODSize=256,MinMagFilter="aniso",MipFilter="point")
++TextureLODGroups=(Group=TEXTUREGROUP_Terrain_Weightmap,LODBias=0,NumStreamedMips=-1,MipGenSettings=TMGS_SimpleAverage,MinLODSize=1,MaxLODSize=256,MinMagFilter="aniso",MipFilter="point")
++TextureLODGroups=(Group=TEXTUREGROUP_Bokeh,LODBias=0,NumStreamedMips=-1,MipGenSettings=TMGS_SimpleAverage,MinLODSize=1,MaxLODSize=256,MinMagFilter="linear",MipFilter="linear")
++TextureLODGroups=(Group=TEXTUREGROUP_IESLightProfile,LODBias=0,NumStreamedMips=-1,MipGenSettings=TMGS_SimpleAverage,MinLODSize=1,MaxLODSize=256,MinMagFilter="Aniso",MipFilter="Point")
++TextureLODGroups=(Group=TEXTUREGROUP_Pixels2D,LODBias=0,NumStreamedMips=-1,MipGenSettings=TMGS_SimpleAverage,MinLODSize=1,MaxLODSize=256,MinMagFilter="point",MipFilter="point")
++TextureLODGroups=(Group=TEXTUREGROUP_HierarchicalLOD,LODBias=0,NumStreamedMips=-1,MipGenSettings=TMGS_SimpleAverage,MinLODSize=1,MaxLODSize=256,MinMagFilter="Aniso",MipFilter="Point")
 
 [Mac DeviceProfile]
 DeviceType=Mac
