Remove unused cvar causing mac automation cis warning
Mirrors CL# 3894945
#rnx

diff --git a/Engine/Config/BaseDeviceProfiles.ini b/Engine/Config/BaseDeviceProfiles.ini
--- a/Engine/Config/BaseDeviceProfiles.ini
+++ b/Engine/Config/BaseDeviceProfiles.ini
@@ -136,7 +136,6 @@ BaseProfileName=
 +CVars=r.MetalPixelParameterSize=50
 +CVars=slate.AbsoluteIndices=1
 +CVars=r.MorphTarget.Mode=0
-+CVars=r.Metal.ManualVertexFetch=1
 
 [iPad2 DeviceProfile]
 DeviceType=IOS
@@ -511,7 +510,6 @@ DeviceType=Mac
 BaseProfileName=
 +CVars=r.Shaders.ZeroInitialise=1
 +CVars=r.Shaders.BoundsChecking=1
-+CVars=r.Metal.ManualVertexFetch=1
 
 [MacClient DeviceProfile]
 DeviceType=Mac
