Add required redirectorfor moving classes from TimeManagement to  CoreUObject
#rnx

diff --git a/Engine/Config/BaseEngine.ini b/Engine/Config/BaseEngine.ini
--- a/Engine/Config/BaseEngine.ini
+++ b/Engine/Config/BaseEngine.ini
@@ -570,6 +570,12 @@ bUseStreamingPause=false
 +ClassRedirects=(OldName="/Script/AssetScriptingUtilities.SkeletalMeshUtilitiesLibrary", NewName="/Script/EditorScriptingUtilities.EditorSkeletalMeshLibrary")
 +FunctionRedirects=(OldName="StaticMeshUtilitiesLibrary.GetLODScreenSizes",NewName="EditorStaticMeshLibrary.GetLodScreenSizes")
 
++StructRedirects=(OldName="FrameNumber",NewName="/Script/CoreUObject.FrameNumber")
++StructRedirects=(OldName="FrameRate",NewName="/Script/CoreUObject.FrameRate")
++StructRedirects=(OldName="FrameTime",NewName="/Script/CoreUObject.FrameTime")
++StructRedirects=(OldName="QualifiedFrameTime",NewName="/Script/CoreUObject.QualifiedFrameTime")
++StructRedirects=(OldName="Timecode",NewName="/Script/CoreUObject.Timecode")
+
 [CoreUObject.Metadata]
 ; Note: UnrealHeaderTool should be rerun after making changes to MetadataRedirects to catch any keys specified in class headers
 +MetadataRedirects=(OldKey="K2Protected", NewKey="BlueprintProtected")
