This is needed because that enum value would have been written into the EditorPerProjectSettings for anyone that modified editor style, and then cause mysterious cook warnings

diff --git a/Engine/Config/BaseEngine.ini b/Engine/Config/BaseEngine.ini
--- a/Engine/Config/BaseEngine.ini
+++ b/Engine/Config/BaseEngine.ini
@@ -512,6 +512,7 @@ bUseStreamingPause=false
 +FunctionRedirects=(OldName="/Script/HeadMountedDisplay.HeadMountedDisplayFunctionLibrary.AddDeviceVisualizationComponent",NewName="/Script/HeadMountedDisplay.XRAssetFunctionLibrary.AddDeviceVisualizationComponentBlocking")
 +FunctionRedirects=(OldName="/Script/HeadMountedDisplay.HeadMountedDisplayFunctionLibrary.AddNamedDeviceVisualizationComponent",NewName="/Script/HeadMountedDisplay.XRAssetFunctionLibrary.AddNamedDeviceVisualizationComponentBlocking")
 +EnumRedirects=(OldName="ESimulationSpace",ValueChanges=(("RootBoneSpace", "BaseBoneSpace")))
++EnumRedirects=(OldName="EColorVisionDeficiency", NewName="/Script/SlateCore.EColorVisionDeficiency", ValueChanges=(("CVD_NormalVision", "NormalVision"), ("CVD_Deuteranomly", "NormalVision"), ("CVD_Deuteranopia", "Deuteranope"), ("CVD_Protanomly", "NormalVision"), ("CVD_Protanopia", "Protanope"), ("CVD_Tritanomaly", "NormalVision"), ("CVD_Tritanopia", "Tritanope"), ("CVD_Achromatopsia", "NormalVision"))
 
 +ClassRedirects=(OldName="NavigationSystem",NewName="/Script/NavigationSystem.NavigationSystemV1")
 +ClassRedirects=(OldName="NavMeshBoundsVolume",NewName="/Script/NavigationSystem.NavMeshBoundsVolume")
