Add Mali G72 device profile for Android
#ue4
#android

diff --git a/Engine/Config/BaseDeviceProfiles.ini b/Engine/Config/BaseDeviceProfiles.ini
--- a/Engine/Config/BaseDeviceProfiles.ini
+++ b/Engine/Config/BaseDeviceProfiles.ini
@@ -73,6 +73,7 @@
 +DeviceProfileNameAndTypes=Android_Mali_T8xx,Android
 +DeviceProfileNameAndTypes=Android_Mali_T8xx_No_Vulkan,Android
 +DeviceProfileNameAndTypes=Android_Mali_G71,Android
++DeviceProfileNameAndTypes=Android_Mali_G72,Android
 ; @todo Lumin: Once we get the IniPlatformParent stuff, this should move into a LuminBaseDeviceProfiles.ini, and make type be Lumin
 +DeviceProfileNameAndTypes=Lumin,Android
 
@@ -353,6 +354,7 @@ MatchProfile=(Profile="Android_PowerVR54x",Match=((SourceType=SRC_GpuFamily,Comp
 +MatchProfile=(Profile="Android_Mali_T8xx_No_Vulkan",Match=((SourceType=SRC_GpuFamily,CompareType=CMP_Regex,MatchString="^Mali\\-T8"),(SourceType=SRC_VulkanVersion,CompareType=CMP_Regex,MatchString="^1\\.0\\.([0-9]+)"),(SourceType=SRC_PreviousRegexMatch,CompareType=CMP_Less,MatchString="11")))
 +MatchProfile=(Profile="Android_Mali_T8xx",Match=((SourceType=SRC_GpuFamily,CompareType=CMP_Regex,MatchString="^Mali\\-T8")))
 +MatchProfile=(Profile="Android_Mali_G71",Match=((SourceType=SRC_GpuFamily,CompareType=CMP_Regex,MatchString="^Mali\\-G71")))
++MatchProfile=(Profile="Android_Mali_G72",Match=((SourceType=SRC_GpuFamily,CompareType=CMP_Regex,MatchString="^Mali\\-G72")))
 +MatchProfile=(Profile="Android_Mid",Match=((SourceType=SRC_GlVersion,CompareType=CMP_Regex,MatchString="^OpenGL ES 3\\.")))
 +MatchProfile=(Profile="Android_Low",Match=((SourceType=SRC_GlVersion,CompareType=CMP_Regex,MatchString="^OpenGL ES 2\\.")))
 
@@ -518,6 +520,10 @@ BaseProfileName=Android_High
 DeviceType=Android
 BaseProfileName=Android_High
 
+[Android_Mali_G72 DeviceProfile]
+DeviceType=Android
+BaseProfileName=Android_High
+
 [Android_TegraK1 DeviceProfile]
 DeviceType=Android
 BaseProfileName=Android_High
