Fixed: Android Rendering issues on ARPG in Vulkan
Disabled occlusion queries on Adreno5xx Vulkan Android 7.0

diff --git a/Engine/Config/BaseDeviceProfiles.ini b/Engine/Config/BaseDeviceProfiles.ini
--- a/Engine/Config/BaseDeviceProfiles.ini
+++ b/Engine/Config/BaseDeviceProfiles.ini
@@ -337,6 +337,7 @@ MatchProfile=(Profile="Android_PowerVR54x",Match=((SourceType=SRC_GpuFamily,Comp
 +MatchProfile=(Profile="Android_Adreno4xx_Nexus6",Match=((SourceType=SRC_DeviceModel,CompareType=CMP_Equal,MatchString="Nexus 6")))
 +MatchProfile=(Profile="Android_Adreno4xx",Match=((SourceType=SRC_GpuFamily,CompareType=CMP_Regex,MatchString="Adreno \\(TM\\) 4[0-9][0-9]")))
 +MatchProfile=(Profile="Android_Adreno51x",Match=((SourceType=SRC_GpuFamily,CompareType=CMP_Regex,MatchString="Adreno \\(TM\\) 51[0-9]")))
++MatchProfile=(Profile="Android_Adreno5xx_Vulkan_Nougat",Match=((SourceType=SRC_GpuFamily,CompareType=CMP_Regex,MatchString="Adreno \\(TM\\) 5[0-9][0-9]"),,(SourceType=SRC_VulkanAvailable,CompareType=CMP_Equal,MatchString="true"),(SourceType=SRC_AndroidVersion,CompareType=CMP_Regex,MatchString="7\\.")))
 +MatchProfile=(Profile="Android_Adreno5xx",Match=((SourceType=SRC_GpuFamily,CompareType=CMP_Regex,MatchString="Adreno \\(TM\\) 5[0-9][0-9]")))
 +MatchProfile=(Profile="Android_Adreno6xx",Match=((SourceType=SRC_GpuFamily,CompareType=CMP_Regex,MatchString="Adreno \\(TM\\) 6[0-9][0-9]")))
 +MatchProfile=(Profile="Android_TegraK1",Match=((SourceType=SRC_GpuFamily,CompareType=CMP_Equal,MatchString="NVIDIA Tegra"),(SourceType=SRC_GlVersion,CompareType=CMP_Regex,MatchString="^OpenGL ES 3\\.")))
@@ -435,6 +436,12 @@ DeviceType=Android
 BaseProfileName=Android_High
 +CVars=r.DisjointTimerQueries=1
 
+[Android_Adreno5xx_Vulkan_Nougat DeviceProfile]
+DeviceType=Android
+BaseProfileName=Android_Adreno5xx
+; vulkan occlusion queries do not work on this device/android version atm 
++CVars=r.AllowOcclusionQueries=0
+
 [Android_Adreno6xx DeviceProfile]
 DeviceType=Android
 BaseProfileName=Android_High
