Merged missed component remapping entry for the ARKit 2.0 integration

diff --git a/Engine/Config/BaseEngine.ini b/Engine/Config/BaseEngine.ini
--- a/Engine/Config/BaseEngine.ini
+++ b/Engine/Config/BaseEngine.ini
@@ -573,6 +573,8 @@ bUseStreamingPause=false
 +ClassRedirects=(OldName="/Script/AssetScriptingUtilities.SkeletalMeshUtilitiesLibrary", NewName="/Script/EditorScriptingUtilities.EditorSkeletalMeshLibrary")
 +FunctionRedirects=(OldName="StaticMeshUtilitiesLibrary.GetLODScreenSizes",NewName="EditorStaticMeshLibrary.GetLodScreenSizes")
 
++ClassRedirects=(OldName="AppleARKitFaceMeshComponent",NewName="/Script/AppleARKitFaceSupport.AppleARKitFaceMeshComponent")
+
 +StructRedirects=(OldName="FrameNumber",NewName="/Script/CoreUObject.FrameNumber")
 +StructRedirects=(OldName="FrameRate",NewName="/Script/CoreUObject.FrameRate")
 +StructRedirects=(OldName="FrameTime",NewName="/Script/CoreUObject.FrameTime")
