Fixed broken NavLinkHostInterface redirection in BaseEngine.ini #UE4

diff --git a/Engine/Config/BaseEngine.ini b/Engine/Config/BaseEngine.ini
--- a/Engine/Config/BaseEngine.ini
+++ b/Engine/Config/BaseEngine.ini
@@ -548,7 +548,7 @@ bUseStreamingPause=false
 +ClassRedirects=(OldName="NavigationGraphNodeComponent",NewName="/Script/NavigationSystem.NavigationGraphNodeComponent")
 +ClassRedirects=(OldName="NavigationPathGenerator",NewName="/Script/NavigationSystem.NavigationPathGenerator")
 +ClassRedirects=(OldName="NavLinkCustomInterface",NewName="/Script/NavigationSystem.NavLinkCustomInterface")
-+ClassRedirects=(OldName="NavLinkHostInterface",NewName="/Script/NavigationSystemConfig.NavLinkHostInterface")
++ClassRedirects=(OldName="NavLinkHostInterface",NewName="/Script/NavigationSystem.NavLinkHostInterface")
 +ClassRedirects=(OldName="NavNodeInterface",NewName="/Script/NavigationSystem.NavNodeInterface")
 +ClassRedirects=(OldName="NavLinkProxy",NewName="/Script/AIModule.NavLinkProxy")
 +ClassRedirects=(OldName="NavigationSystemConfig", NewName="/Script/NavigationSystem.NavigationSystemModuleConfig")
