Increase default Metal shader version for iOS to 1.1, which requires iOS 9. This is the only way to avoid a Metal compiler error with the Xcode 9 tools (Xcode 10 seems OK).

diff --git a/Engine/Config/BaseEngine.ini b/Engine/Config/BaseEngine.ini
--- a/Engine/Config/BaseEngine.ini
+++ b/Engine/Config/BaseEngine.ini
@@ -1710,6 +1710,7 @@ bGenerateCrashReportSymbols=false
 bAutomaticSigning=false
 UseFastIntrinsics=False
 EnableMathOptimisations=True
+MaxShaderLanguageVersion=1
 
 ; These are the defaults for Android settings, and they need to be in the .ini since UBT reads the .ini settings, without instantiating the class
 [/Script/AndroidRuntimeSettings.AndroidRuntimeSettings]
