MagicLeap 0.12.0 devices need to do something special to enable vulkan, so just go back to OpenGL as default for 4.20 release.  4.21 on will be 0.15.0 SDK and will be Vulkan only

diff --git a/Engine/Config/BaseEngine.ini b/Engine/Config/BaseEngine.ini
--- a/Engine/Config/BaseEngine.ini
+++ b/Engine/Config/BaseEngine.ini
@@ -1781,7 +1781,7 @@ bBuildWithNvTegraGfxDebugger=False
 Certificate=(FilePath="")
 IconModelPath=(Path="../../Build/Lumin/Resources/Model")
 IconPortalPath=(Path="../../Build/Lumin/Resources/Portal")
-bUseVulkan=true
+bUseVulkan=false
 VersionCode=1
 bInternetRequired=true
 +AppPrivileges=Headpose
