REL - Blacklist memory leaking Nvidia driver

diff --git a/Engine/Config/BaseHardware.ini b/Engine/Config/BaseHardware.ini
--- a/Engine/Config/BaseHardware.ini
+++ b/Engine/Config/BaseHardware.ini
@@ -47,6 +47,11 @@ SuggestedDriverVersion="378.92"
 ; UE-35288, GPU hangs in 4.13 binary release.
 +Blacklist=(DriverVersion="==372.70", Reason="This driver version has many known stability issues")
 
+; UE-58673, GPU memory leak and hitching.
++Blacklist=(DriverVersion="==388.31", Reason="This driver is leaking memory latest known working driver was 388.13")
++Blacklist=(DriverVersion="==388.43", Reason="This driver is leaking memory latest known working driver was 388.13")
++Blacklist=(DriverVersion="==388.71", Reason="This driver is leaking memory latest known working driver was 388.13")
+
 ;-----------------------------------------------------------------------------------------------------------------
 
 [GPU_AMD]
