Disable Aftermath until a fix for a memory leak is found

diff --git a/Engine/Config/BaseEngine.ini b/Engine/Config/BaseEngine.ini
--- a/Engine/Config/BaseEngine.ini
+++ b/Engine/Config/BaseEngine.ini
@@ -2028,7 +2028,7 @@ bEnableLiveLinkForFaceTracking=true
 LiveLinkPublishingPort=11111
 
 [/Script/Engine.RendererSettings]
-r.GPUCrashDebugging=True
+r.GPUCrashDebugging=false
 
 [Messaging]
 bAllowDelayedMessaging=false
\ No newline at end of file
