Disregard For GC pool will now be enabled by default in cooked builds

diff --git a/Engine/Config/BaseEngine.ini b/Engine/Config/BaseEngine.ini
--- a/Engine/Config/BaseEngine.ini
+++ b/Engine/Config/BaseEngine.ini
@@ -682,7 +682,7 @@ s.MaxPackageSummarySize=16384
 +FixedBootOrder=/Script/Engine/Default__DeviceProfileManager
 
 [/Script/Engine.GarbageCollectionSettings]
-gc.MaxObjectsNotConsideredByGC=0
+gc.MaxObjectsNotConsideredByGC=1
 gc.SizeOfPermanentObjectPool=0
 gc.FlushStreamingOnGC=0
 gc.NumRetriesBeforeForcingGC=10
