Remove obsolete cvars from ini files

diff --git a/Engine/Config/BaseDeviceProfiles.ini b/Engine/Config/BaseDeviceProfiles.ini
--- a/Engine/Config/BaseDeviceProfiles.ini
+++ b/Engine/Config/BaseDeviceProfiles.ini
@@ -133,8 +133,6 @@ BaseProfileName=
 +CVars=r.PostProcessAAQuality=0
 ; Not sure on this one - may need the whole D Buffer stuff going
 +CVars=r.Decal.StencilSizeThreshold=-1
-+CVars=r.MetalVertexParameterSize=50
-+CVars=r.MetalPixelParameterSize=50
 +CVars=slate.AbsoluteIndices=1
 +CVars=r.MorphTarget.Mode=0
 
@@ -283,8 +281,6 @@ BaseProfileName=IOS
 +CVars=r.MobileContentScaleFactor=1.5
 +CVars=g.TimeoutForBlockOnRenderFence=3000000
 +CVars=r.Decal.StencilSizeThreshold=0.1
-+CVars=r.MetalVertexParameterSize=1024
-+CVars=r.MetalPixelParameterSize=1024
 +CVars=r.MetalComputeParameterSize=1024
 +CVars=r.EarlyZPass=3
 +CVars=r.EarlyZPassMovable=1
