Linux: remove GL3 RHI from the default settings (UE-64269).

diff --git a/Engine/Config/BaseEngine.ini b/Engine/Config/BaseEngine.ini
--- a/Engine/Config/BaseEngine.ini
+++ b/Engine/Config/BaseEngine.ini
@@ -1919,11 +1919,9 @@ MinimumOSVersion=MSOS_Vista
 ; Possible values for architecture: X86_64UnknownLinuxGnu, ArmUnknownLinuxGnueabihf, AArch64UnknownLinuxGnueabi, I686UnknownLinuxGnu
 TargetArchitecture=X86_64UnknownLinuxGnu
 
-; When neither -vulkan, -opengl3 nor -opengl4 are passed on the commandline, the engine will default to the first targeted RHI.
+; When neither -vulkan nor -opengl4 are passed on the commandline, the engine will default to the first targeted RHI.
 +TargetedRHIs=SF_VULKAN_SM5
 +TargetedRHIs=GLSL_430
-+TargetedRHIs=GLSL_150
-+TargetedRHIs=SF_VULKAN_ES31
 
  
 [/Script/MacTargetPlatform.MacTargetSettings]
