Fix incorrect path to engine DDC when running an executable built with a unique build environment that's more than one path deep. Add an %ENGINEDIR% variable that can be expanded from config files.

diff --git a/Engine/Config/BaseEngine.ini b/Engine/Config/BaseEngine.ini
--- a/Engine/Config/BaseEngine.ini
+++ b/Engine/Config/BaseEngine.ini
@@ -1483,11 +1483,11 @@ Root=(Type=KeyLength, Length=120, Inner=AsyncPut)
 AsyncPut=(Type=AsyncPut, Inner=Hierarchy)
 Hierarchy=(Type=Hierarchical, Inner=Boot, Inner=Pak, Inner=EnginePak, Inner=Local, Inner=Shared)
 Boot=(Type=Boot, Filename="%GAMEDIR%DerivedDataCache/Boot.ddc", MaxCacheSize=512)
-Local=(Type=FileSystem, ReadOnly=false, Clean=false, Flush=false, PurgeTransient=true, DeleteUnused=true, UnusedFileAge=34, FoldersToClean=-1, Path=../../../Engine/DerivedDataCache, EnvPathOverride=UE-LocalDataCachePath, EditorOverrideSetting=LocalDerivedDataCache)
+Local=(Type=FileSystem, ReadOnly=false, Clean=false, Flush=false, PurgeTransient=true, DeleteUnused=true, UnusedFileAge=34, FoldersToClean=-1, Path=%ENGINEDIR%DerivedDataCache, EnvPathOverride=UE-LocalDataCachePath, EditorOverrideSetting=LocalDerivedDataCache)
 Shared=(Type=FileSystem, ReadOnly=false, Clean=false, Flush=false, DeleteUnused=true, UnusedFileAge=10, FoldersToClean=10, MaxFileChecksPerSec=1, Path=?EpicDDC, EnvPathOverride=UE-SharedDataCachePath, EditorOverrideSetting=SharedDerivedDataCache)
 AltShared=(Type=FileSystem, ReadOnly=true, Clean=false, Flush=false, DeleteUnused=true, UnusedFileAge=23, FoldersToClean=10, MaxFileChecksPerSec=1, Path=?EpicDDC2, EnvPathOverride=UE-SharedDataCachePath2)
 Pak=(Type=ReadPak, Filename="%GAMEDIR%DerivedDataCache/DDC.ddp")
-EnginePak=(Type=ReadPak, Filename=../../../Engine/DerivedDataCache/DDC.ddp)
+EnginePak=(Type=ReadPak, Filename=%ENGINEDIR%DerivedDataCache/DDC.ddp)
 
 [DerivedDataBackendGraph_Fill_Seattle]
 MinimumDaysToKeepFile=7
@@ -1495,10 +1495,10 @@ Root=(Type=KeyLength, Length=120, Inner=AsyncPut)
 AsyncPut=(Type=AsyncPut, Inner=Hierarchy)
 Hierarchy=(Type=Hierarchical, Inner=Boot, Inner=Pak, Inner=EnginePak, Inner=Local, Inner=Seattle)
 Boot=(Type=Boot, Filename="%GAMEDIR%DerivedDataCache/Boot.ddc", MaxCacheSize=512)
-Local=(Type=FileSystem, ReadOnly=false, Clean=false, Flush=false, PurgeTransient=true, DeleteUnused=true, UnusedFileAge=34, FoldersToClean=-1, Path=../../../Engine/DerivedDataCache)
+Local=(Type=FileSystem, ReadOnly=false, Clean=false, Flush=false, PurgeTransient=true, DeleteUnused=true, UnusedFileAge=34, FoldersToClean=-1, Path=%ENGINEDIR%DerivedDataCache)
 Seattle=(Type=FileSystem, ReadOnly=false, Clean=false, Flush=false, DeleteUnused=true, UnusedFileAge=23, FoldersToClean=10, MaxFileChecksPerSec=1, Path=?EpicSeaDDC, EnvPathOverride=UE-SharedDataCachePath_Seattle)
 Pak=(Type=ReadPak, Filename="%GAMEDIR%DerivedDataCache/DDC.ddp")
-EnginePak=(Type=ReadPak, Filename=../../../Engine/DerivedDataCache/DDC.ddp)
+EnginePak=(Type=ReadPak, Filename=%ENGINEDIR%DerivedDataCache/DDC.ddp)
 
 [InstalledDerivedDataBackendGraph]
 MinimumDaysToKeepFile=7
@@ -1519,7 +1519,7 @@ Root=(Type=KeyLength, Length=120, Inner=AsyncPut)
 AsyncPut=(Type=AsyncPut, Inner=Hierarchy)
 Hierarchy=(Type=Hierarchical, Inner=Boot, Inner=Pak, Inner=Local)
 Boot=(Type=Boot, Filename="%GAMEDIR%DerivedDataCache/Boot.ddc", MaxCacheSize=512)
-Local=(Type=FileSystem, ReadOnly=false, Clean=false, Flush=false, PurgeTransient=true, DeleteUnused=true, UnusedFileAge=34, FoldersToClean=-1, Path=../../../Engine/DerivedDataCache, EnvPathOverride=UE-LocalDataCachePath)
+Local=(Type=FileSystem, ReadOnly=false, Clean=false, Flush=false, PurgeTransient=true, DeleteUnused=true, UnusedFileAge=34, FoldersToClean=-1, Path=%ENGINEDIR%DerivedDataCache, EnvPathOverride=UE-LocalDataCachePath)
 Pak=(Type=ReadPak, Filename="%GAMEDIR%DerivedDataCache/DDC.ddp")
 
 [CreatePak]
@@ -1528,7 +1528,7 @@ Root=(Type=KeyLength, Length=120, Inner=AsyncPut)
 AsyncPut=(Type=AsyncPut, Inner=Hierarchy)
 Hierarchy=(Type=Hierarchical, Inner=Boot, Inner=PakWrite, Inner=PakRead, Inner=Local, Inner=Shared)
 Boot=(Type=Boot, Filename="%GAMEDIR%DerivedDataCache/Boot.ddc", MaxCacheSize=512)
-Local=(Type=FileSystem, ReadOnly=false, Clean=false, Flush=false, PurgeTransient=true, DeleteUnused=true, UnusedFileAge=34, FoldersToClean=-1, Path=../../../Engine/DerivedDataCache)
+Local=(Type=FileSystem, ReadOnly=false, Clean=false, Flush=false, PurgeTransient=true, DeleteUnused=true, UnusedFileAge=34, FoldersToClean=-1, Path=%ENGINEDIR%DerivedDataCache)
 Shared=(Type=FileSystem, ReadOnly=false, Clean=false, Flush=false, DeleteUnused=false, UnusedFileAge=23, FoldersToClean=-1, Path=?EpicDDC, EnvPathOverride=UE-SharedDataCachePath)
 AltShared=(Type=FileSystem, ReadOnly=true, Clean=false, Flush=false, DeleteUnused=false, UnusedFileAge=23, FoldersToClean=-1, Path=?EpicDDC2, EnvPathOverride=UE-SharedDataCachePath2)
 PakRead=(Type=ReadPak, Filename="%GAMEDIR%DerivedDataCache/DDC.ddp")
@@ -1540,7 +1540,7 @@ Root=(Type=KeyLength, Length=120, Inner=AsyncPut)
 AsyncPut=(Type=AsyncPut, Inner=Hierarchy)
 Hierarchy=(Type=Hierarchical, Inner=Boot, Inner=EnginePak, Inner=PakWrite, Inner=PakRead, Inner=Local, Inner=Shared)
 Boot=(Type=Boot, Filename="%GAMEDIR%DerivedDataCache/Boot.ddc", MaxCacheSize=512)
-Local=(Type=FileSystem, ReadOnly=false, Clean=false, Flush=false, PurgeTransient=true, DeleteUnused=true, UnusedFileAge=34, FoldersToClean=-1, Path=../../../Engine/DerivedDataCache, EnvPathOverride=UE-LocalDataCachePath)
+Local=(Type=FileSystem, ReadOnly=false, Clean=false, Flush=false, PurgeTransient=true, DeleteUnused=true, UnusedFileAge=34, FoldersToClean=-1, Path=%ENGINEDIR%DerivedDataCache, EnvPathOverride=UE-LocalDataCachePath)
 Shared=(Type=FileSystem, ReadOnly=false, Clean=false, Flush=false, DeleteUnused=false, UnusedFileAge=23, FoldersToClean=-1, Path=?EpicDDC, EnvPathOverride=UE-SharedDataCachePath)
 AltShared=(Type=FileSystem, ReadOnly=true, Clean=false, Flush=false, DeleteUnused=false, UnusedFileAge=23, FoldersToClean=-1, Path=?EpicDDC2, EnvPathOverride=UE-SharedDataCachePath2)
 EnginePak=(Type=ReadPak, Filename=../../../Engine/DerivedDataCache/Compressed.ddp, Compressed=true)
@@ -1553,10 +1553,10 @@ Root=(Type=KeyLength, Length=120, Inner=AsyncPut)
 AsyncPut=(Type=AsyncPut, Inner=Hierarchy)
 Hierarchy=(Type=Hierarchical, Inner=Boot, Inner=PakWrite, Inner=Local, Inner=Shared)
 Boot=(Type=Boot, Filename="%GAMEDIR%DerivedDataCache/Boot.ddc", MaxCacheSize=512)
-Local=(Type=FileSystem, ReadOnly=false, Clean=false, Flush=false, PurgeTransient=true, DeleteUnused=true, UnusedFileAge=34, FoldersToClean=-1, Path=../../../Engine/DerivedDataCache, EnvPathOverride=UE-LocalDataCachePath)
+Local=(Type=FileSystem, ReadOnly=false, Clean=false, Flush=false, PurgeTransient=true, DeleteUnused=true, UnusedFileAge=34, FoldersToClean=-1, Path=%ENGINEDIR%DerivedDataCache, EnvPathOverride=UE-LocalDataCachePath)
 Shared=(Type=FileSystem, ReadOnly=false, Clean=false, Flush=false, DeleteUnused=false, UnusedFileAge=23, FoldersToClean=-1, Path=?EpicDDC, EnvPathOverride=UE-SharedDataCachePath)
 AltShared=(Type=FileSystem, ReadOnly=true, Clean=false, Flush=false, DeleteUnused=false, UnusedFileAge=23, FoldersToClean=-1, Path=?EpicDDC2, EnvPathOverride=UE-SharedDataCachePath2)
-PakWrite=(Type=WritePak, Filename=../../../Engine/DerivedDataCache/Compressed.ddp, Compressed=true)
+PakWrite=(Type=WritePak, Filename=%ENGINEDIR%DerivedDataCache/Compressed.ddp, Compressed=true)
 
 [CreateInstalledEnterprisePak]
 MinimumDaysToKeepFile=7
@@ -1564,7 +1564,7 @@ Root=(Type=KeyLength, Length=120, Inner=AsyncPut)
 AsyncPut=(Type=AsyncPut, Inner=Hierarchy)
 Hierarchy=(Type=Hierarchical, Inner=Boot, Inner=PakWrite, Inner=Local, Inner=Shared)
 Boot=(Type=Boot, Filename="%GAMEDIR%DerivedDataCache/Boot.ddc", MaxCacheSize=512)
-Local=(Type=FileSystem, ReadOnly=false, Clean=false, Flush=false, PurgeTransient=true, DeleteUnused=true, UnusedFileAge=34, FoldersToClean=-1, Path=../../../Engine/DerivedDataCache, EnvPathOverride=UE-LocalDataCachePath)
+Local=(Type=FileSystem, ReadOnly=false, Clean=false, Flush=false, PurgeTransient=true, DeleteUnused=true, UnusedFileAge=34, FoldersToClean=-1, Path=%ENGINEDIR%DerivedDataCache, EnvPathOverride=UE-LocalDataCachePath)
 Shared=(Type=FileSystem, ReadOnly=false, Clean=false, Flush=false, DeleteUnused=false, UnusedFileAge=23, FoldersToClean=-1, Path=?EpicDDC, EnvPathOverride=UE-SharedDataCachePath)
 AltShared=(Type=FileSystem, ReadOnly=true, Clean=false, Flush=false, DeleteUnused=false, UnusedFileAge=23, FoldersToClean=-1, Path=?EpicDDC2, EnvPathOverride=UE-SharedDataCachePath2)
 PakWrite=(Type=WritePak, Filename=../../../Enterprise/DerivedDataCache/Compressed.ddp, Compressed=true)
@@ -1574,7 +1574,7 @@ MinimumDaysToKeepFile=7
 Root=(Type=KeyLength, Length=120, Inner=AsyncPut)
 AsyncPut=(Type=AsyncPut, Inner=Hierarchy)
 Hierarchy=(Type=Hierarchical, Inner=Local, Inner=Project, Inner=Shared)
-Local=(Type=FileSystem, ReadOnly=false, Clean=false, Flush=false, PurgeTransient=true, DeleteUnused=true, UnusedFileAge=34, FoldersToClean=-1, Path=../../../Engine/DerivedDataCache, EnvPathOverride=UE-LocalDataCachePath)
+Local=(Type=FileSystem, ReadOnly=false, Clean=false, Flush=false, PurgeTransient=true, DeleteUnused=true, UnusedFileAge=34, FoldersToClean=-1, Path=%ENGINEDIR%DerivedDataCache, EnvPathOverride=UE-LocalDataCachePath)
 Project=(Type=FileSystem, ReadOnly=false, Clean=false, Flush=true, PurgeTransient=true, DeleteUnused=true, UnusedFileAge=34, FoldersToClean=-1, Path=%GAMEDIR%ProjectDerivedData)
 Shared=(Type=FileSystem, ReadOnly=false, Clean=false, Flush=false, DeleteUnused=false, UnusedFileAge=23, FoldersToClean=-1, Path=?EpicDDC, EnvPathOverride=UE-SharedDataCachePath)
 AltShared=(Type=FileSystem, ReadOnly=true, Clean=false, Flush=false, DeleteUnused=false, UnusedFileAge=23, FoldersToClean=-1, Path=?EpicDDC2, EnvPathOverride=UE-SharedDataCachePath2)
diff --git a/Engine/Source/Runtime/Core/Private/Misc/ConfigCacheIni.cpp b/Engine/Source/Runtime/Core/Private/Misc/ConfigCacheIni.cpp
--- a/Engine/Source/Runtime/Core/Private/Misc/ConfigCacheIni.cpp
+++ b/Engine/Source/Runtime/Core/Private/Misc/ConfigCacheIni.cpp
@@ -92,6 +92,7 @@ static const FConfigExpansion* MatchExpansions(const TCHAR* PotentialVariable)
 	{
 		FConfigExpansion(TEXT("%GAME%"), FString(FApp::GetProjectName())),
 		FConfigExpansion(TEXT("%GAMEDIR%"), FPaths::ProjectDir()),
+		FConfigExpansion(TEXT("%ENGINEDIR%"), FPaths::EngineDir()),
 		FConfigExpansion(TEXT("%ENGINEUSERDIR%"), FPaths::EngineUserDir()),
 		FConfigExpansion(TEXT("%ENGINEVERSIONAGNOSTICUSERDIR%"), FPaths::EngineVersionAgnosticUserDir()),
 		FConfigExpansion(TEXT("%APPSETTINGSDIR%"), GetApplicationSettingsDirNormalized()),
