ConsoleVariables.ini - added MeshDrawCommands debugging and cleaned up a bit.
#rnx

diff --git a/Engine/Config/ConsoleVariables.ini b/Engine/Config/ConsoleVariables.ini
--- a/Engine/Config/ConsoleVariables.ini
+++ b/Engine/Config/ConsoleVariables.ini
@@ -37,12 +37,15 @@
 ; Uncomment to dump shaders in the Saved folder
 ; Warning: leaving this on for a while will fill your hard drive with many small files and folders
 ;r.DumpShaderDebugInfo=1
-; When this is enabled, dumped shader paths will get collapsed (in the cases where paths are longer than the OS's max)
-;r.DumpShaderDebugShortNames=1
 ; When this is enabled, when dumping shaders an additional file to use with ShaderCompilerWorker -direct mode will be generated
 ;r.DumpShaderDebugWorkerCommandLine=1
-; Uncomment to enable parallel rendering at startup
-;r.RHICmdBypass=0
+
+; Uncomment to disable parallel rendering
+;r.RHICmdBypass=1
+; Uncomment to disable parallel mesh draw command setup
+;r.MeshDrawCommands.ParallelPassSetup=0
+; Uncomment to disable cached mesh draw commands
+;r.MeshDrawCommands.UseCachedCommands=0
 
 ; Uncomment to get render graph executing passes as they get created to easily debug crashes caused by pass wiring logic.
 ;r.RDG.ImmediateMode=1
