adding support for iPadMini5 and iPadAir3

diff --git a/Engine/Config/BaseDeviceProfiles.ini b/Engine/Config/BaseDeviceProfiles.ini
--- a/Engine/Config/BaseDeviceProfiles.ini
+++ b/Engine/Config/BaseDeviceProfiles.ini
@@ -44,6 +44,8 @@
 +DeviceProfileNameAndTypes=iPad6,IOS
 +DeviceProfileNameAndTypes=iPadPro11,IOS
 +DeviceProfileNameAndTypes=iPadPro3_129,IOS
++DeviceProfileNameAndTypes=iPadAir3,IOS
++DeviceProfileNameAndTypes=iPadMini5,IOS
 +DeviceProfileNameAndTypes=AppleTV,IOS
 +DeviceProfileNameAndTypes=AppleTV4K,IOS
 +DeviceProfileNameAndTypes=HTML5,HTML5
@@ -373,6 +375,14 @@ BaseProfileName=IPadPro
 DeviceType=IOS
 BaseProfileName=IPadPro
 
+[iPadAir3 DeviceProfile]
+DeviceType=IOS
+BaseProfileName=IPadPro
+
+[iPadMini5 DeviceProfile]
+DeviceType=IOS
+BaseProfileName=IPadPro
+
 [/Script/AndroidDeviceProfileSelector.AndroidDeviceProfileMatchingRules]
 MatchProfile=(Profile="Android_PowerVR54x",Match=((SourceType=SRC_GpuFamily,CompareType=CMP_Regex,MatchString="PowerVR SGX 54[0-9]")))
 +MatchProfile=(Profile="Android_PowerVR_G6xxx",Match=((SourceType=SRC_GpuFamily,CompareType=CMP_Regex,MatchString="PowerVR Rogue G6[0-9]+")))
diff --git a/Engine/Config/BaseEditorPerProjectUserSettings.ini b/Engine/Config/BaseEditorPerProjectUserSettings.ini
--- a/Engine/Config/BaseEditorPerProjectUserSettings.ini
+++ b/Engine/Config/BaseEditorPerProjectUserSettings.ini
@@ -274,7 +274,9 @@ LastExecutedLaunchPlatform=
 +TabletScreenResolutions=(Description="iPad Pro 9.7-inch",Width=768,Height=1024,AspectRatio="3:4",bCanSwapAspectRatio=true,ProfileName="iPadPro97")
 +TabletScreenResolutions=(Description="iPad (6th gen.)",Width=768,Height=1024,AspectRatio="3:4",bCanSwapAspectRatio=true,ProfileName="iPad6")
 +TabletScreenResolutions=(Description="iPad (5th gen.)",Width=768,Height=1024,AspectRatio="3:4",bCanSwapAspectRatio=true,ProfileName="iPad5")
++TabletScreenResolutions=(Description="iPad Air 3",Width=768,Height=1024,AspectRatio="3:4",bCanSwapAspectRatio=true,ProfileName="iPadAir3")
 +TabletScreenResolutions=(Description="iPad Air 2",Width=768,Height=1024,AspectRatio="3:4",bCanSwapAspectRatio=true,ProfileName="iPadAir2")
++TabletScreenResolutions=(Description="iPad Mini 5",Width=768,Height=1024,AspectRatio="3:4",bCanSwapAspectRatio=true,ProfileName="iPadMini5")
 +TabletScreenResolutions=(Description="iPad Mini 4",Width=768,Height=1024,AspectRatio="3:4",bCanSwapAspectRatio=true,ProfileName="iPadMini4")
 +TabletScreenResolutions=(Description="LG G Pad X 8.0",Width=768,Height=1366,AspectRatio="9:16",bCanSwapAspectRatio=true)
 +TabletScreenResolutions=(Description="Asus Zenpad 3s 10",Width=768,Height=1366,AspectRatio="9:16",bCanSwapAspectRatio=true)
@@ -859,4 +861,4 @@ IncludeSmoothLODCheck=false
 VertexPreviewSize=6
 
 [/Script/UndoHistory.UndoHistorySettings]
-bShowTransactionDetails=false
\ No newline at end of file
+bShowTransactionDetails=false
diff --git a/Engine/Source/Runtime/ApplicationCore/Private/IOS/IOSPlatformApplicationMisc.cpp b/Engine/Source/Runtime/ApplicationCore/Private/IOS/IOSPlatformApplicationMisc.cpp
--- a/Engine/Source/Runtime/ApplicationCore/Private/IOS/IOSPlatformApplicationMisc.cpp
+++ b/Engine/Source/Runtime/ApplicationCore/Private/IOS/IOSPlatformApplicationMisc.cpp
@@ -251,7 +251,7 @@ void FIOSPlatformApplicationMisc::ClipboardPaste(class FString& Result)
 EScreenPhysicalAccuracy FIOSPlatformApplicationMisc::ComputePhysicalScreenDensity(int32& ScreenDensity)
 {
 	FPlatformMisc::EIOSDevice Device = FPlatformMisc::GetIOSDeviceType();
-	static_assert( FPlatformMisc::EIOSDevice::IOS_Unknown == 38, "Every device needs to be handled here." );
+	static_assert( FPlatformMisc::EIOSDevice::IOS_Unknown == 40, "Every device needs to be handled here." );
 
 	ScreenDensity = 0;
 	EScreenPhysicalAccuracy Accuracy = EScreenPhysicalAccuracy::Unknown;
@@ -289,6 +289,7 @@ EScreenPhysicalAccuracy FIOSPlatformApplicationMisc::ComputePhysicalScreenDensit
 	case FPlatformMisc::IOS_IPadMini:
 	case FPlatformMisc::IOS_IPadMini2: // also the iPadMini3
 	case FPlatformMisc::IOS_IPadMini4:
+    case FPlatformMisc::IOS_IPadMini5:
 		ScreenDensity = 401;
 		Accuracy = EScreenPhysicalAccuracy::Truth;
 		break;
@@ -299,6 +300,7 @@ EScreenPhysicalAccuracy FIOSPlatformApplicationMisc::ComputePhysicalScreenDensit
 	case FPlatformMisc::IOS_IPad6:
 	case FPlatformMisc::IOS_IPadAir:
 	case FPlatformMisc::IOS_IPadAir2:
+    case FPlatformMisc::IOS_IPadAir3:
 	case FPlatformMisc::IOS_IPadPro_97:
 		ScreenDensity = 264;
 		Accuracy = EScreenPhysicalAccuracy::Truth;
diff --git a/Engine/Source/Runtime/Core/Private/IOS/IOSPlatformMisc.cpp b/Engine/Source/Runtime/Core/Private/IOS/IOSPlatformMisc.cpp
--- a/Engine/Source/Runtime/Core/Private/IOS/IOSPlatformMisc.cpp
+++ b/Engine/Source/Runtime/Core/Private/IOS/IOSPlatformMisc.cpp
@@ -494,7 +494,17 @@ FIOSPlatformMisc::EIOSDevice FIOSPlatformMisc::GetIOSDeviceType()
 				DeviceType = IOS_IPadPro3_129;
 			}
 		}
-
+        else if (Major == 11)
+        {
+            if (Minor <= 2)
+            {
+                DeviceType = IOS_IPadMini5;
+            }
+            else
+            {
+                DeviceType = IOS_IPadAir3;
+            }
+        }
 		// Default to highest settings currently available for any future device
 		else if (Major >= 9)
 		{
diff --git a/Engine/Source/Runtime/Core/Public/IOS/IOSPlatformMisc.h b/Engine/Source/Runtime/Core/Public/IOS/IOSPlatformMisc.h
--- a/Engine/Source/Runtime/Core/Public/IOS/IOSPlatformMisc.h
+++ b/Engine/Source/Runtime/Core/Public/IOS/IOSPlatformMisc.h
@@ -141,6 +141,8 @@ struct CORE_API FIOSPlatformMisc : public FApplePlatformMisc
 		IOS_IPad6,
 		IOS_IPadPro_11,
 		IOS_IPadPro3_129,
+        IOS_IPadAir3,
+        IOS_IPadMini5,
 		IOS_Unknown,
 	};
 
@@ -188,6 +190,8 @@ struct CORE_API FIOSPlatformMisc : public FApplePlatformMisc
 			TEXT("IPad6"),
 			TEXT("IPadPro11"),
 			TEXT("IPadPro3_129"),
+            TEXT("IPadAir3"),
+            TEXT("IPadMini5"),
 			TEXT("Unknown"),
 		};
 		static_assert((sizeof(IOSDeviceNames) / sizeof(IOSDeviceNames[0])) == ((int32)IOS_Unknown + 1), "Mismatched IOSDeviceNames and EIOSDevice.");
