HTML5 - 4.22.1 - update browser paths for launch on menu

#rnx

diff --git a/Engine/Config/BaseEngine.ini b/Engine/Config/BaseEngine.ini
--- a/Engine/Config/BaseEngine.ini
+++ b/Engine/Config/BaseEngine.ini
@@ -2089,17 +2089,21 @@ LLVMWasmBackend=False
 ; common browser locations - add your custom locations here
 [/Script/HTML5PlatformEditor.HTML5Browsers]
 ; windows
-+BrowserLauncher=(BrowserName="Nightly(64bit)", BrowserPath=(FilePath="C:/Program Files/Nightly/firefox.exe"))
-+BrowserLauncher=(BrowserName="Nightly",        BrowserPath=(FilePath="C:/Program Files (x86)/Nightly/firefox.exe"))
-+BrowserLauncher=(BrowserName="Firefox(64bit)", BrowserPath=(FilePath="C:/Program Files/Mozilla Firefox/firefox.exe"))
-+BrowserLauncher=(BrowserName="Firefox",        BrowserPath=(FilePath="C:/Program Files (x86)/Mozilla Firefox/firefox.exe"))
-+BrowserLauncher=(BrowserName="Chrome",         BrowserPath=(FilePath="C:/Program Files (x86)/Google/Chrome/Application/chrome.exe"))
-
++BrowserLauncher=(BrowserName="FF Nightly(32bit)", BrowserPath=(FilePath="C:/Program Files (x86)/Firefox Nightly/firefox.exe"))
++BrowserLauncher=(BrowserName="FF Dev(32bit)",     BrowserPath=(FilePath="C:/Program Files (x86)/Firefox Developer Edition/firefox.exe"))
++BrowserLauncher=(BrowserName="Firefox(32bit)",    BrowserPath=(FilePath="C:/Program Files (x86)/Mozilla Firefox/firefox.exe"))
++BrowserLauncher=(BrowserName="FF Nightly",        BrowserPath=(FilePath="C:/Program Files/Firefox Nightly/firefox.exe"))
++BrowserLauncher=(BrowserName="FF Dev",            BrowserPath=(FilePath="C:/Program Files/Firefox Developer Edition/firefox.exe"))
++BrowserLauncher=(BrowserName="Firefox",           BrowserPath=(FilePath="C:/Program Files/Mozilla Firefox/firefox.exe"))
+; old FF Nightly path...
++BrowserLauncher=(BrowserName="FF Nightly(32bit)", BrowserPath=(FilePath="C:/Program Files (x86)/Nightly/firefox.exe"))
++BrowserLauncher=(BrowserName="FF Nightly",        BrowserPath=(FilePath="C:/Program Files/Nightly/firefox.exe"))
+; chrome
++BrowserLauncher=(BrowserName="Chrome",            BrowserPath=(FilePath="C:/Program Files (x86)/Google/Chrome/Application/chrome.exe"))
 ; Chrome Canary -- edit path to canary installation location
-;+BrowserLauncher=(BrowserName="Canary",         BrowserPath=(FilePath="[REDACTED_USER_HOME]/AppData/Local/Google/Chrome SxS/Application/chrome.exe"))
-
-; Edge -- uncomment following if you have Windows10
-;+BrowserLauncher=(BrowserName="Edge",           BrowserPath=(FilePath="start microsoft-edge:"))
++BrowserLauncher=(BrowserName="Canary",            BrowserPath=(FilePath="[REDACTED_USER_HOME]/AppData/Local/Google/Chrome SxS/Application/chrome.exe"))
+; Edge -- will be autodetected...
+;+BrowserLauncher=(BrowserName="Edge",             BrowserPath=(FilePath="start microsoft-edge:"))
 
 ; osx
 +BrowserLauncher=(BrowserName="Safari",  BrowserPath=(FilePath="/Applications/Safari.app"))
diff --git a/Engine/Source/Developer/HTML5/HTML5TargetPlatform/Private/HTML5TargetPlatform.cpp b/Engine/Source/Developer/HTML5/HTML5TargetPlatform/Private/HTML5TargetPlatform.cpp
--- a/Engine/Source/Developer/HTML5/HTML5TargetPlatform/Private/HTML5TargetPlatform.cpp
+++ b/Engine/Source/Developer/HTML5/HTML5TargetPlatform/Private/HTML5TargetPlatform.cpp
@@ -350,7 +350,13 @@ void FHTML5TargetPlatform::PopulateDevices(TArray<FString>& DeviceMaps, FString
 		if( FParse::Value( *It, TEXT( "BrowserName=" ), DeviceName ) && !DeviceName.IsEmpty() &&
 			FParse::Value( *It, TEXT( "BrowserPath=(FilePath=" ), DevicePath ) && !DevicePath.IsEmpty() )
 		{
-
+#if PLATFORM_WINDOWS
+			if ( DeviceName.StartsWith( TEXT( "Canary" ) ) )
+			{
+				FString UserName = FPlatformMisc::GetEnvironmentVariable(TEXT("USRENAME"));
+				DevicePath.Replace( TEXT("your_username"), *UserName );
+			}
+#endif
 			if (FPlatformFileManager::Get().GetPlatformFile().FileExists(*DevicePath) ||
 			    FPlatformFileManager::Get().GetPlatformFile().DirectoryExists(*DevicePath))
 			{
diff --git a/Engine/Source/Programs/HTML5/HTML5LaunchHelper/Program.cs b/Engine/Source/Programs/HTML5/HTML5LaunchHelper/Program.cs
--- a/Engine/Source/Programs/HTML5/HTML5LaunchHelper/Program.cs
+++ b/Engine/Source/Programs/HTML5/HTML5LaunchHelper/Program.cs
@@ -1,4 +1,4 @@
-﻿// Copyright 1998-2019 Epic Games, Inc. All Rights Reserved.
+// Copyright 1998-2019 Epic Games, Inc. All Rights Reserved.
 
 using System;
 using System.Net;
@@ -454,14 +454,14 @@ namespace HTML5LaunchHelper
 		static void SpawnBrowserAndBlock(Arguments Args)
 		{
 			// Browsers can be multiprocess programs (Chrome, basically)
-			bool bMultiprocessBrowser = Args.Browser.Contains("chrome");
+			bool bMultiprocessBrowser = Args.Browser.Contains("chrome") || Args.Browser.Contains("firefox");
 			// so we need to catch spawning of other child processes. The trick is
 			// they aren't really child-processes at all. There appears to be no real binding between the two,
 			// So we kind of fudged it a bit here.
 			var PrevProcesses = Process.GetProcesses();
 			var FirstProcess = SpawnBrowserProcess(Args.Browser, Args.BrowserCommandLine);
 			ProcessesToWatch.Add(FirstProcess);
-			var ProcName = FirstProcess.ProcessName;
+			var ProcName = FirstProcess.HasExited ? Args.Browser : FirstProcess.ProcessName;
 
 			// We should now have a list of processes to watch to exit.
 			// Loop over the calling WaitForExit() until the list is empty.
@@ -554,7 +554,7 @@ namespace HTML5LaunchHelper
 
 		static int Main(string[] args)
 		{
-			System.Console.WriteLine("Version: 20170906"); // date: YYYYMMDD - needed to help figure out what version QA is running...
+			System.Console.WriteLine("Version: 20190327"); // date: YYYYMMDD - needed to help figure out what version QA is running...
 			var Args = new Arguments();
 			if (Args.Parse(args))
 			{
