Console autocomplete for starting and stopping the network profiler

diff --git a/Engine/Config/BaseInput.ini b/Engine/Config/BaseInput.ini
--- a/Engine/Config/BaseInput.ini
+++ b/Engine/Config/BaseInput.ini
@@ -196,6 +196,8 @@ MaxScrollbackSize=1024
 +ManualAutoCompleteList=(Command="Automation",Desc="Run an automation command (e.g., Automation RunTests TestName)")
 +ManualAutoCompleteList=(Command="CsvProfile Start",Desc="Start CSV profiling.")
 +ManualAutoCompleteList=(Command="CsvProfile Stop",Desc="Stop CSV profiling.")
++ManualAutoCompleteList=(Command="NetProfile Enable",Desc="Start network profiling.")
++ManualAutoCompleteList=(Command="NetProfile Disable",Desc="Stop network profiling.")
 BackgroundOpacityPercentage=85.000000
 InputColor=(R=230,G=230,B=230,A=255)
 HistoryColor=(R=180,G=180,B=180,A=255)
