Update base input so auto-complete is supported for audio commands.

diff --git a/Engine/Config/BaseInput.ini b/Engine/Config/BaseInput.ini
--- a/Engine/Config/BaseInput.ini
+++ b/Engine/Config/BaseInput.ini
@@ -120,30 +120,57 @@ MaxScrollbackSize=1024
 +ManualAutoCompleteList=(Command="ListAnims",Desc="Lists info for all animations")
 +ManualAutoCompleteList=(Command="ListSkeletalMeshes",Desc="Lists info for all skeletal meshes")
 +ManualAutoCompleteList=(Command="ListStaticMeshes",Desc="Lists info for all static meshes")
-+ManualAutoCompleteList=(Command="AudioMemReport",Desc="Lists info for audio memory")
 +ManualAutoCompleteList=(Command="InvestigateTexture",Desc="Shows streaming info about the specified texture")
 ;placed here so we can type res<tab> to restartlevel as we often do that
 +ManualAutoCompleteList=(Command="RestartLevel",Desc="Restarts the level")
-;from Audio.cpp
 +ManualAutoCompleteList=(Command="Module List",Desc="Lists all known modules")
 +ManualAutoCompleteList=(Command="Module Load",Desc="Attempts to load the specified module name")
 +ManualAutoCompleteList=(Command="Module Unload",Desc="Unloads the specified module name")
 +ManualAutoCompleteList=(Command="Module Reload",Desc="Reloads the specified module name, unloading it first if needed")
 +ManualAutoCompleteList=(Command="Module Recompile",Desc="Attempts to recompile a module, first unloading it if needed")
 +ManualAutoCompleteList=(Command="HotReload",Desc="<UObject DLL Hot Reload> Attempts to recompile a UObject DLL and reload it on the fly")
+
+
+; From AudioDevice.cpp
++ManualAutoCompleteList=(Command="AudioDebugSound",Desc"<filter> Rejects new USoundBase requests where the sound name does not contain the provided filter")
++ManualAutoCompleteList=(Command="AudioGetDynamicSoundVolume",Desc"Gets volume for given sound type ('Cue' or 'Wave') with provided name")
++ManualAutoCompleteList=(Command="AudioMemReport",Desc="Lists info for audio memory")
++ManualAutoCompleteList=(Command="AudioMixerDebugSound",Desc"<filter> With new mixer enabled, rejects new USoundBase requests where the sound name does not contain the provided filter")
++ManualAutoCompleteList=(Command="AudioResetAllDynamicSoundVolumes",Desc"Resets all dynamic volumes to unity")
++ManualAutoCompleteList=(Command="AudioResetDynamicSoundVolume",Desc"Resets volume for given sound type ('Cue' or 'Wave') with provided name to unity")
++ManualAutoCompleteList=(Command="AudioSetDynamicSoundVolume",Desc"Name=<name> Type=<type> Vol=<vol> Sets volume for given sound type ('Cue' or 'Wave') with provided name")
++ManualAutoCompleteList=(Command="AudioSoloSoundClass",Desc"<name> Solos USoundClass")
++ManualAutoCompleteList=(Command="AudioSoloSoundCue",Desc"<name> Solos USoundCue")
++ManualAutoCompleteList=(Command="AudioSoloSoundWave",Desc"<name> Solos USoundWave")
++ManualAutoCompleteList=(Command="ClearSoloAudio",Desc"Clears solo'ed object")
++ManualAutoCompleteList=(Command="DisableLPF",Desc"Disables low-pass filter")
++ManualAutoCompleteList=(Command="DisableEQFilter",Desc="Disables EQ")
++ManualAutoCompleteList=(Command="DisableRadio",Desc"Disables radio effect")
++ManualAutoCompleteList=(Command="DumpSoundInfo",Desc"Dumps sound information to log")
++ManualAutoCompleteList=(Command="EnableRadio",Desc"Enables radio effect")
++ManualAutoCompleteList=(Command="IsolateDryAudio",Desc"Isolates dry audio")
++ManualAutoCompleteList=(Command="IsolateReverb",Desc"Isolates reverb")
++ManualAutoCompleteList=(Command="ListAudioComponents",Desc="Dumps a detailed list of all AudioComponent objects")
++ManualAutoCompleteList=(Command="ListSoundClasses",Desc="Lists a summary of loaded sound collated by class")
++ManualAutoCompleteList=(Command="ListSoundClassVolumes",Desc"Lists all sound class volumes")
++ManualAutoCompleteList=(Command="ListSoundDurations",Desc"Lists durations of all active sounds")
 +ManualAutoCompleteList=(Command="ListSounds",Desc="Lists all the loaded sounds and their memory footprint")
 +ManualAutoCompleteList=(Command="ListWaves",Desc="List the WaveInstances and whether they have a source")
-+ManualAutoCompleteList=(Command="ListSoundClasses",Desc="Lists a summary of loaded sound collated by class")
-+ManualAutoCompleteList=(Command="ListAudioComponents",Desc="Dumps a detailed list of all AudioComponent objects")
-+ManualAutoCompleteList=(Command="ListSoundDurations",Desc=)
++ManualAutoCompleteList=(Command="PlayAllPIEAudio",Desc"Toggls whether or not all devices should play their audio")
 +ManualAutoCompleteList=(Command="PlaySoundCue",Desc="Lists a summary of loaded sound collated by class")
-+ManualAutoCompleteList=(Command="PlaySoundWave",Desc=)
-+ManualAutoCompleteList=(Command="SetBaseSoundMix",Desc="<MixName>")
-+ManualAutoCompleteList=(Command="DisableLowPassFilter",Desc=)
-+ManualAutoCompleteList=(Command="DisableEQFilter",Desc=)
-+ManualAutoCompleteList=(Command="IsolateDryAudio",Desc=)
-+ManualAutoCompleteList=(Command="IsolateReverb",Desc=)
++ManualAutoCompleteList=(Command="PlaySoundWave",Desc="<name> Plays the given soundwave")
 +ManualAutoCompleteList=(Command="ResetSoundState",Desc="Resets volumes to default and removes test filters")
++ManualAutoCompleteList=(Command="SetBaseSoundMix",Desc="<MixName> Sets the base sound mix")
++ManualAutoCompleteList=(Command="ShowSoundClassHierarchy",Desc"")
++ManualAutoCompleteList=(Command="SoloAudio",Desc"Solos the audio device associated with the parent world")
++ManualAutoCompleteList=(Command="SoundClassFixup",Desc"Deprecated")
++ManualAutoCompleteList=(Command="TestLFEBleed",Desc"Sets LPF to max for all sources")
++ManualAutoCompleteList=(Command="TestLPF",Desc"Sets LPF to max for all sources")
++ManualAutoCompleteList=(Command="TestStereoBleed",Desc"Test bleeding stereo sounds fully to the rear speakers")
++ManualAutoCompleteList=(Command="ToggleHRTFForAll",Desc"Toggles whether or not HRTF spatialization is enabled for all")
++ManualAutoCompleteList=(Command="ToggleSpatExt",Desc"Toggles enablement of the Spatial Audio Extension")
+
+
 +ManualAutoCompleteList=(Command="DisableAllScreenMessages",Desc="Disables all on-screen warnings/messages")
 +ManualAutoCompleteList=(Command="EnableAllScreenMessages",Desc="Enables all on-screen warnings/messages")
 +ManualAutoCompleteList=(Command="ToggleAllScreenMessages",Desc="Toggles display state of all on-screen warnings/messages")
