Fix copy pasta'ed description of cmd

diff --git a/Engine/Config/BaseInput.ini b/Engine/Config/BaseInput.ini
--- a/Engine/Config/BaseInput.ini
+++ b/Engine/Config/BaseInput.ini
@@ -157,7 +157,7 @@ MaxScrollbackSize=1024
 +ManualAutoCompleteList=(Command="ListSounds",Desc="Lists all the loaded sounds and their memory footprint")
 +ManualAutoCompleteList=(Command="ListWaves",Desc="List the WaveInstances and whether they have a source")
 +ManualAutoCompleteList=(Command="PlayAllPIEAudio",Desc"Toggls whether or not all devices should play their audio")
-+ManualAutoCompleteList=(Command="PlaySoundCue",Desc="Lists a summary of loaded sound collated by class")
++ManualAutoCompleteList=(Command="PlaySoundCue",Desc="Plays the given soundcue")
 +ManualAutoCompleteList=(Command="PlaySoundWave",Desc="<name> Plays the given soundwave")
 +ManualAutoCompleteList=(Command="ResetSoundState",Desc="Resets volumes to default and removes test filters")
 +ManualAutoCompleteList=(Command="SetBaseSoundMix",Desc="<MixName> Sets the base sound mix")
