Fix Linux having bypass enabled by default (UE-73337).

diff --git a/Engine/Config/BaseDeviceProfiles.ini b/Engine/Config/BaseDeviceProfiles.ini
--- a/Engine/Config/BaseDeviceProfiles.ini
+++ b/Engine/Config/BaseDeviceProfiles.ini
@@ -710,6 +710,12 @@ DeviceType=Linux
 BaseProfileName=Linux
 MeshLODSettings=
 TextureLODSettings=
++CVars=r.RHICmdBypass=0
+
+[LinuxClient DeviceProfile]
+DeviceType=Linux
+BaseProfileName=Linux
++CVars=r.RHICmdBypass=0
 
 [LinuxServer DeviceProfile]
 DeviceType=Linux
