Fix missing '=' copy pasta in input ManualAutoCompleteList commands

diff --git a/Engine/Config/BaseInput.ini b/Engine/Config/BaseInput.ini
--- a/Engine/Config/BaseInput.ini
+++ b/Engine/Config/BaseInput.ini
@@ -132,43 +132,43 @@ MaxScrollbackSize=1024
 
 
 ; From AudioDevice.cpp
-+ManualAutoCompleteList=(Command="AudioDebugSound",Desc"<filter> Rejects new USoundBase requests where the sound name does not contain the provided filter")
-+ManualAutoCompleteList=(Command="AudioGetDynamicSoundVolume",Desc"Gets volume for given sound type ('Cue' or 'Wave') with provided name")
++ManualAutoCompleteList=(Command="AudioDebugSound",Desc="<filter> Rejects new USoundBase requests where the sound name does not contain the provided filter")
++ManualAutoCompleteList=(Command="AudioGetDynamicSoundVolume",Desc="Gets volume for given sound type ('Cue' or 'Wave') with provided name")
 +ManualAutoCompleteList=(Command="AudioMemReport",Desc="Lists info for audio memory")
-+ManualAutoCompleteList=(Command="AudioMixerDebugSound",Desc"<filter> With new mixer enabled, rejects new USoundBase requests where the sound name does not contain the provided filter")
-+ManualAutoCompleteList=(Command="AudioResetAllDynamicSoundVolumes",Desc"Resets all dynamic volumes to unity")
-+ManualAutoCompleteList=(Command="AudioResetDynamicSoundVolume",Desc"Resets volume for given sound type ('Cue' or 'Wave') with provided name to unity")
-+ManualAutoCompleteList=(Command="AudioSetDynamicSoundVolume",Desc"Name=<name> Type=<type> Vol=<vol> Sets volume for given sound type ('Cue' or 'Wave') with provided name")
-+ManualAutoCompleteList=(Command="AudioSoloSoundClass",Desc"<name> Solos USoundClass")
-+ManualAutoCompleteList=(Command="AudioSoloSoundCue",Desc"<name> Solos USoundCue")
-+ManualAutoCompleteList=(Command="AudioSoloSoundWave",Desc"<name> Solos USoundWave")
-+ManualAutoCompleteList=(Command="ClearSoloAudio",Desc"Clears solo'ed object")
-+ManualAutoCompleteList=(Command="DisableLPF",Desc"Disables low-pass filter")
++ManualAutoCompleteList=(Command="AudioMixerDebugSound",Desc="<filter> With new mixer enabled, rejects new USoundBase requests where the sound name does not contain the provided filter")
++ManualAutoCompleteList=(Command="AudioResetAllDynamicSoundVolumes",Desc="Resets all dynamic volumes to unity")
++ManualAutoCompleteList=(Command="AudioResetDynamicSoundVolume",Desc="Resets volume for given sound type ('Cue' or 'Wave') with provided name to unity")
++ManualAutoCompleteList=(Command="AudioSetDynamicSoundVolume",Desc="Name=<name> Type=<type> Vol=<vol> Sets volume for given sound type ('Cue' or 'Wave') with provided name")
++ManualAutoCompleteList=(Command="AudioSoloSoundClass",Desc="<name> Solos USoundClass")
++ManualAutoCompleteList=(Command="AudioSoloSoundCue",Desc="<name> Solos USoundCue")
++ManualAutoCompleteList=(Command="AudioSoloSoundWave",Desc="<name> Solos USoundWave")
++ManualAutoCompleteList=(Command="ClearSoloAudio",Desc="Clears solo'ed object")
++ManualAutoCompleteList=(Command="DisableLPF",Desc="Disables low-pass filter")
 +ManualAutoCompleteList=(Command="DisableEQFilter",Desc="Disables EQ")
-+ManualAutoCompleteList=(Command="DisableRadio",Desc"Disables radio effect")
-+ManualAutoCompleteList=(Command="DumpSoundInfo",Desc"Dumps sound information to log")
-+ManualAutoCompleteList=(Command="EnableRadio",Desc"Enables radio effect")
-+ManualAutoCompleteList=(Command="IsolateDryAudio",Desc"Isolates dry audio")
-+ManualAutoCompleteList=(Command="IsolateReverb",Desc"Isolates reverb")
++ManualAutoCompleteList=(Command="DisableRadio",Desc="Disables radio effect")
++ManualAutoCompleteList=(Command="DumpSoundInfo",Desc="Dumps sound information to log")
++ManualAutoCompleteList=(Command="EnableRadio",Desc="Enables radio effect")
++ManualAutoCompleteList=(Command="IsolateDryAudio",Desc="Isolates dry audio")
++ManualAutoCompleteList=(Command="IsolateReverb",Desc="Isolates reverb")
 +ManualAutoCompleteList=(Command="ListAudioComponents",Desc="Dumps a detailed list of all AudioComponent objects")
 +ManualAutoCompleteList=(Command="ListSoundClasses",Desc="Lists a summary of loaded sound collated by class")
-+ManualAutoCompleteList=(Command="ListSoundClassVolumes",Desc"Lists all sound class volumes")
-+ManualAutoCompleteList=(Command="ListSoundDurations",Desc"Lists durations of all active sounds")
++ManualAutoCompleteList=(Command="ListSoundClassVolumes",Desc="Lists all sound class volumes")
++ManualAutoCompleteList=(Command="ListSoundDurations",Desc="Lists durations of all active sounds")
 +ManualAutoCompleteList=(Command="ListSounds",Desc="Lists all the loaded sounds and their memory footprint")
 +ManualAutoCompleteList=(Command="ListWaves",Desc="List the WaveInstances and whether they have a source")
-+ManualAutoCompleteList=(Command="PlayAllPIEAudio",Desc"Toggls whether or not all devices should play their audio")
++ManualAutoCompleteList=(Command="PlayAllPIEAudio",Desc="Toggls whether or not all devices should play their audio")
 +ManualAutoCompleteList=(Command="PlaySoundCue",Desc="Plays the given soundcue")
 +ManualAutoCompleteList=(Command="PlaySoundWave",Desc="<name> Plays the given soundwave")
 +ManualAutoCompleteList=(Command="ResetSoundState",Desc="Resets volumes to default and removes test filters")
 +ManualAutoCompleteList=(Command="SetBaseSoundMix",Desc="<MixName> Sets the base sound mix")
-+ManualAutoCompleteList=(Command="ShowSoundClassHierarchy",Desc"")
-+ManualAutoCompleteList=(Command="SoloAudio",Desc"Solos the audio device associated with the parent world")
-+ManualAutoCompleteList=(Command="SoundClassFixup",Desc"Deprecated")
-+ManualAutoCompleteList=(Command="TestLFEBleed",Desc"Sets LPF to max for all sources")
-+ManualAutoCompleteList=(Command="TestLPF",Desc"Sets LPF to max for all sources")
-+ManualAutoCompleteList=(Command="TestStereoBleed",Desc"Test bleeding stereo sounds fully to the rear speakers")
-+ManualAutoCompleteList=(Command="ToggleHRTFForAll",Desc"Toggles whether or not HRTF spatialization is enabled for all")
-+ManualAutoCompleteList=(Command="ToggleSpatExt",Desc"Toggles enablement of the Spatial Audio Extension")
++ManualAutoCompleteList=(Command="ShowSoundClassHierarchy",Desc="")
++ManualAutoCompleteList=(Command="SoloAudio",Desc="Solos the audio device associated with the parent world")
++ManualAutoCompleteList=(Command="SoundClassFixup",Desc="Deprecated")
++ManualAutoCompleteList=(Command="TestLFEBleed",Desc="Sets LPF to max for all sources")
++ManualAutoCompleteList=(Command="TestLPF",Desc="Sets LPF to max for all sources")
++ManualAutoCompleteList=(Command="TestStereoBleed",Desc="Test bleeding stereo sounds fully to the rear speakers")
++ManualAutoCompleteList=(Command="ToggleHRTFForAll",Desc="Toggles whether or not HRTF spatialization is enabled for all")
++ManualAutoCompleteList=(Command="ToggleSpatExt",Desc="Toggles enablement of the Spatial Audio Extension")
 
 
 +ManualAutoCompleteList=(Command="DisableAllScreenMessages",Desc="Disables all on-screen warnings/messages")
