* Disabled the distance-scaled camera speed setting by default.

diff --git a/Engine/Config/BaseEditorPerProjectUserSettings.ini b/Engine/Config/BaseEditorPerProjectUserSettings.ini
--- a/Engine/Config/BaseEditorPerProjectUserSettings.ini
+++ b/Engine/Config/BaseEditorPerProjectUserSettings.ini
@@ -319,7 +319,7 @@ bInvertRightMouseDollyYAxis=False
 bUseAbsoluteTranslation=True
 bLevelStreamingVolumePrevis=False
 bUseUE3OrbitControls=False
-bUseDistanceScaledCameraSpeed=True
+bUseDistanceScaledCameraSpeed=False
 bOrbitCameraAroundSelection=True
 ;Scroll gesture direction
 ScrollGestureDirectionFor3DViewports=UseSystemSetting
