Fix early pak downloader not including global shader archive when using shared cooks.

diff --git a/Engine/Config/BaseGame.ini b/Engine/Config/BaseGame.ini
--- a/Engine/Config/BaseGame.ini
+++ b/Engine/Config/BaseGame.ini
@@ -91,7 +91,7 @@ bSkipMovies=False
 +EarlyDownloaderPakFileFiles=...\Config\...\*.ini
 +EarlyDownloaderPakFileFiles=...\Config\*.ini
 +EarlyDownloaderPakFileFiles=...\Engine\GlobalShaderCache*.bin
-+EarlyDownloaderPakFileFiles=...\Content\ShaderArchive-Global-*.ushaderbytecode
++EarlyDownloaderPakFileFiles=...\Content\ShaderArchive-Global*.ushaderbytecode
 +EarlyDownloaderPakFileFiles=...\Content\Slate\*.*
 +EarlyDownloaderPakFileFiles=...\Content\Slate\...\*.*
 +EarlyDownloaderPakFileFiles=...\*.upluginmanifest
