Niagara - Detail Level in base scalability ini

diff --git a/Engine/Config/BaseScalability.ini b/Engine/Config/BaseScalability.ini
--- a/Engine/Config/BaseScalability.ini
+++ b/Engine/Config/BaseScalability.ini
@@ -389,6 +389,7 @@ r.SSS.Quality=0
 r.SSS.HalfRes=1
 r.EmitterSpawnRateScale=0.125
 r.ParticleLightQuality=0
+r.DetailLevel=0
 
 [EffectsQuality@1]
 r.TranslucencyLightingVolumeDim=32
@@ -404,6 +405,7 @@ r.SSS.Quality=0
 r.SSS.HalfRes=1
 r.EmitterSpawnRateScale=0.25
 r.ParticleLightQuality=0
+r.DetailLevel=1
 
 [EffectsQuality@2]
 r.TranslucencyLightingVolumeDim=48
@@ -419,6 +421,7 @@ r.SSS.Quality=-1
 r.SSS.HalfRes=1
 r.EmitterSpawnRateScale=0.5
 r.ParticleLightQuality=1
+r.DetailLevel=2
 
 [EffectsQuality@3]
 r.TranslucencyLightingVolumeDim=64
@@ -434,6 +437,7 @@ r.SSS.Quality=1
 r.SSS.HalfRes=1
 r.EmitterSpawnRateScale=1.0
 r.ParticleLightQuality=2
+r.DetailLevel=3
 
 [EffectsQuality@Cine]
 r.TranslucencyLightingVolumeDim=64
@@ -449,6 +453,7 @@ r.SSS.Quality=1
 r.SSS.HalfRes=0
 r.EmitterSpawnRateScale=1.0
 r.ParticleLightQuality=2
+r.DetailLevel=4
 
 ;-----------------------------------------------------------------------------------------------------------------
 
