Sequencer: Setting default filmback preset to 16:9 DSLR. Updated groundtruth image for LevelSequecneTestNiagara After Runup.

diff --git a/Engine/Config/BaseEngine.ini b/Engine/Config/BaseEngine.ini
--- a/Engine/Config/BaseEngine.ini
+++ b/Engine/Config/BaseEngine.ini
@@ -2200,7 +2200,7 @@ ServerUpdateCameraTimeout=2.0
 +FilmbackPresets=(Name="APS-C (Canon)",FilmbackSettings=(SensorWidth=22.2,SensorHeight=14.8))
 +FilmbackPresets=(Name="Full Frame DSLR",FilmbackSettings=(SensorWidth=36,SensorHeight=24))
 +FilmbackPresets=(Name="Micro Four Thirds",FilmbackSettings=(SensorWidth=17.3,SensorHeight=13))
-DefaultFilmbackPresetName="16:9 Digital Film"
+DefaultFilmbackPresetName="16:9 DSLR"
 +LensPresets=(Name="12mm Prime f/2.8",LensSettings=(MinFocalLength=12,MaxFocalLength=12,MinFStop=2.8,MaxFStop=22,DiaphragmBladeCount=7))
 +LensPresets=(Name="30mm Prime f/1.4",LensSettings=(MinFocalLength=30,MaxFocalLength=30,MinFStop=1.4,MaxFStop=22,DiaphragmBladeCount=7))
 +LensPresets=(Name="50mm Prime f/1.8",LensSettings=(MinFocalLength=50,MaxFocalLength=50,MinFStop=1.8,MaxFStop=22,DiaphragmBladeCount=7))
