Add DefaultVirtualTextureMaterial to the startup package list
Integrated from //UE4/Dev-Rendering
#rnx

diff --git a/Engine/Config/BaseEngine.ini b/Engine/Config/BaseEngine.ini
--- a/Engine/Config/BaseEngine.ini
+++ b/Engine/Config/BaseEngine.ini
@@ -734,6 +734,7 @@ bFullyCompressStartupPackages=false
 +Package=/Engine/EngineMaterials/DefaultMaterial
 +Package=/Engine/EngineMaterials/DefaultNormal
 +Package=/Engine/EngineMaterials/DefaultPhysicalMaterial
++Package=/Engine/EngineMaterials/DefaultVirtualTextureMaterial
 +Package=/Engine/EngineMaterials/DefaultWhiteGrid
 +Package=/Engine/EngineMaterials/EditorBrushMaterial
 +Package=/Engine/EngineMaterials/EmissiveMeshMaterial
@@ -1070,6 +1071,7 @@ AutoSaveIndex=0
 +PackagesToBeFullyLoadedAtStartup=/Engine/EngineMaterials/WorldGridMaterial
 +PackagesToBeFullyLoadedAtStartup=/Engine/EngineMaterials/DefaultNormal
 +PackagesToBeFullyLoadedAtStartup=/Engine/EngineMaterials/DefaultPhysicalMaterial
++PackagesToBeFullyLoadedAtStartup=/Engine/EngineMaterials/DefaultVirtualTextureMaterial
 +PackagesToBeFullyLoadedAtStartup=/Engine/EngineMaterials/DefaultWhiteGrid
 +PackagesToBeFullyLoadedAtStartup=/Engine/EngineMaterials/EditorBrushMaterial
 +PackagesToBeFullyLoadedAtStartup=/Engine/EngineMaterials/Good64x64TilingNoiseHighFreq
