Remove legacy r.EarlyZPassMovable from iOS device profiles

diff --git a/Engine/Config/BaseDeviceProfiles.ini b/Engine/Config/BaseDeviceProfiles.ini
--- a/Engine/Config/BaseDeviceProfiles.ini
+++ b/Engine/Config/BaseDeviceProfiles.ini
@@ -155,7 +155,6 @@ BaseProfileName=
 +CVars=r.HZBOcclusion=0
 ; Non need for depth prepass
 +CVars=r.EarlyZPass=0
-+CVars=r.EarlyZPassMovable=0
 ; Needs geometry shader support
 +CVars=r.TranslucentLightingVolume=0
 ; Needs geometry shader support
@@ -335,7 +334,6 @@ BaseProfileName=IOS
 +CVars=r.Decal.StencilSizeThreshold=0.1
 +CVars=r.MetalComputeParameterSize=1024
 +CVars=r.EarlyZPass=3
-+CVars=r.EarlyZPassMovable=1
 +CVars=r.PostProcessAAQuality=4
 
 [AppleTV DeviceProfile]
