Making sure the bottom bar safe zone is defined to extend the whole length of the phone.

diff --git a/Engine/Config/BaseDeviceProfiles.ini b/Engine/Config/BaseDeviceProfiles.ini
--- a/Engine/Config/BaseDeviceProfiles.ini
+++ b/Engine/Config/BaseDeviceProfiles.ini
@@ -315,12 +315,12 @@ BaseProfileName=iPhone8
 [iPhoneXSMax DeviceProfile]
 DeviceType=IOS
 BaseProfileName=iPhone8
-+CVars=r.CustomUnsafeZones="(L:free[0,-15][812,15]);(P:fixed[104,0][206,30])"
++CVars=r.CustomUnsafeZones="(L:free[0,-15][896,15]);(P:fixed[104,0][206,30])"
 
 [iPhoneXR DeviceProfile]
 DeviceType=IOS
 BaseProfileName=iPhone8
-+CVars=r.CustomUnsafeZones="(L:free[0,-15][812,15]);(P:fixed[104,0][206,30])"
++CVars=r.CustomUnsafeZones="(L:free[0,-15][896,15]);(P:fixed[104,0][206,30])"
 
 [iPadPro DeviceProfile]
 DeviceType=IOS
