UE-61699 Kindle Fire 5th Gen crashes rotating Remote Session

#Android
#4.21

From //UE4/Dev-Mobile CL4360704

Disabling threaded rendering for Kindle Fire HD 7 5th gen in DefaultDeviceProfiles.ini

diff --git a/Engine/Config/BaseDeviceProfiles.ini b/Engine/Config/BaseDeviceProfiles.ini
--- a/Engine/Config/BaseDeviceProfiles.ini
+++ b/Engine/Config/BaseDeviceProfiles.ini
@@ -516,6 +516,7 @@ BaseProfileName=Android_Mali_4xx
 [Android_Mali_4xx_KindleFire7 DeviceProfile]
 DeviceType=Android
 BaseProfileName=Android_Mali_4xx
++CVars=r.AndroidDisableThreadedRendering=1
 +CVars=r.AndroidDisableThreadedRenderingFirstLoad=1
 
 [Android_Mali_T6xx DeviceProfile]
