Sequencer: Added missing ini change for 4395032

diff --git a/Engine/Config/BaseEngine.ini b/Engine/Config/BaseEngine.ini
--- a/Engine/Config/BaseEngine.ini
+++ b/Engine/Config/BaseEngine.ini
@@ -619,6 +619,10 @@ bUseStreamingPause=false
 +FunctionRedirects=(OldName="GetHandGestureConfidence",NewName="/Script/MagicLeapHandTracking.MagicLeapHandTrackingFunctionLibrary.GetCurrentGestureConfidence")
 +FunctionRedirects=(OldName="GetCurrentGesture",NewName="/Script/MagicLeapHandTracking.MagicLeapHandTrackingFunctionLibrary.GetCurrentGesture")
 
+; LevelSequenceDirectorBlueprint removal
++ClassRedirects=(OldName="LevelSequenceDirectorBlueprint",NewName="/Script/LevelSequence.LegacyLevelSequenceDirectorBlueprint")
++ClassRedirects=(OldName="LevelSequenceDirectorGeneratedClass",NewName="/Script/Engine.BlueprintGeneratedClass")
+
 [CoreUObject.Metadata]
 ; Note: UnrealHeaderTool should be rerun after making changes to MetadataRedirects to catch any keys specified in class headers
 +MetadataRedirects=(OldKey="K2Protected", NewKey="BlueprintProtected")
