Merging using Dev-VR_to_Release-4.21

4392743 UE-63646 MagicLeap 4.21 redirects need updating
-Made redirectors more specific, found a way to do the updates without the MatchSubstring package update flag, to avoid a warning about how it is very slow.

diff --git a/Engine/Config/BaseEngine.ini b/Engine/Config/BaseEngine.ini
--- a/Engine/Config/BaseEngine.ini
+++ b/Engine/Config/BaseEngine.ini
@@ -599,25 +599,25 @@ bUseStreamingPause=false
 ; 4.21
 
 ; MagicLeap Gestures->HandTracking
-+PackageRedirects=(OldName="/Script/MagicLeapGestures", NewName="/Script/MagicLeapHandTracking", MatchSubstring=true)
-+ClassRedirects=(OldName="MagicLeapGesturesFunctionLibrary",NewName="/Script/MagicLeapHandTracking.MagicLeapHandTrackingFunctionLibrary")
-+EnumRedirects=(OldName="EStaticGestures",NewName="/Script/MagicLeapHandTracking.EHandTrackingGesture")
-+EnumRedirects=(OldName="EGestureKeypointsFilterLevel",NewName="/Script/MagicLeapHandTracking.EHandTrackingKeypointFilterLevel")
-+EnumRedirects=(OldName="EGestureRecognitionFilterLevel",NewName="/Script/MagicLeapHandTracking.EHandTrackingGestureFilterLevel")
-+EnumRedirects=(OldName="EGestureTransformSpace",NewName="/Script/MagicLeapHandTracking.EGestureTransformSpace")
++PackageRedirects=(OldName="/Script/MagicLeapGestures", NewName="/Script/MagicLeapHandTracking")
++ClassRedirects=(OldName="/Script/MagicLeapGestures.MagicLeapGesturesFunctionLibrary",NewName="/Script/MagicLeapHandTracking.MagicLeapHandTrackingFunctionLibrary")
++EnumRedirects=(OldName="/Script/MagicLeapGestures.EStaticGestures",NewName="/Script/MagicLeapHandTracking.EHandTrackingGesture")
++EnumRedirects=(OldName="/Script/MagicLeapGestures.EGestureKeypointsFilterLevel",NewName="/Script/MagicLeapHandTracking.EHandTrackingKeypointFilterLevel")
++EnumRedirects=(OldName="/Script/MagicLeapGestures.EGestureRecognitionFilterLevel",NewName="/Script/MagicLeapHandTracking.EHandTrackingGestureFilterLevel")
++EnumRedirects=(OldName="/Script/MagicLeapGestures.EGestureTransformSpace",NewName="/Script/MagicLeapHandTracking.EGestureTransformSpace")
 +EnumRedirects=(OldName="ESceneTextureId",ValueChanges=(("PPI_ShadingModel","PPI_ShadingModelColor"))
 +FunctionRedirects=(OldName="NavigationSystemV1.SimpleMoveToLocation",NewName="AIBlueprintHelperLibrary.SimpleMoveToLocation")
-+FunctionRedirects=(OldName="GetHandCenter",NewName="/Script/MagicLeapHandTracking.MagicLeapHandTrackingFunctionLibrary.GetHandCenter")
-+FunctionRedirects=(OldName="GetHandPointer",NewName="/Script/MagicLeapHandTracking.MagicLeapHandTrackingFunctionLibrary.GetHandIndexFingerTip")
-+FunctionRedirects=(OldName="GetHandSecondary",NewName="/Script/MagicLeapHandTracking.MagicLeapHandTrackingFunctionLibrary.GetHandThumbTip")
-+FunctionRedirects=(OldName="GetHandCenterNormalized",NewName="/Script/MagicLeapHandTracking.MagicLeapHandTrackingFunctionLibrary.GetHandCenterNormalized")
-+FunctionRedirects=(OldName="GetGestureKeypoints",NewName="/Script/MagicLeapHandTracking.MagicLeapHandTrackingFunctionLibrary.GetGestureKeypoints")
-+FunctionRedirects=(OldName="SetConfiguration",NewName="/Script/MagicLeapHandTracking.MagicLeapHandTrackingFunctionLibrary.SetConfiguration")
-+FunctionRedirects=(OldName="GetConfiguration",NewName="/Script/MagicLeapHandTracking.MagicLeapHandTrackingFunctionLibrary.GetConfiguration")
-+FunctionRedirects=(OldName="SetStaticGestureConfidenceThreshold",NewName="/Script/MagicLeapHandTracking.MagicLeapHandTrackingFunctionLibrary.SetStaticGestureConfidenceThreshold")
-+FunctionRedirects=(OldName="GetStaticGestureConfidenceThreshold",NewName="/Script/MagicLeapHandTracking.MagicLeapHandTrackingFunctionLibrary.GetStaticGestureConfidenceThreshold")
-+FunctionRedirects=(OldName="GetHandGestureConfidence",NewName="/Script/MagicLeapHandTracking.MagicLeapHandTrackingFunctionLibrary.GetCurrentGestureConfidence")
-+FunctionRedirects=(OldName="GetCurrentGesture",NewName="/Script/MagicLeapHandTracking.MagicLeapHandTrackingFunctionLibrary.GetCurrentGesture")
++FunctionRedirects=(OldName="/Script/MagicLeapHandTracking.MagicLeapGesturesFunctionLibrary.GetHandCenter",NewName="/Script/MagicLeapHandTracking.MagicLeapHandTrackingFunctionLibrary.GetHandCenter")
++FunctionRedirects=(OldName="/Script/MagicLeapHandTracking.MagicLeapHandTrackingFunctionLibrary.GetHandPointer",NewName="/Script/MagicLeapHandTracking.MagicLeapHandTrackingFunctionLibrary.GetHandIndexFingerTip")
++FunctionRedirects=(OldName="/Script/MagicLeapHandTracking.MagicLeapHandTrackingFunctionLibrary.GetHandSecondary",NewName="/Script/MagicLeapHandTracking.MagicLeapHandTrackingFunctionLibrary.GetHandThumbTip")
++FunctionRedirects=(OldName="/Script/MagicLeapHandTracking.MagicLeapHandTrackingFunctionLibrary.GetHandCenterNormalized",NewName="/Script/MagicLeapHandTracking.MagicLeapHandTrackingFunctionLibrary.GetHandCenterNormalized")
++FunctionRedirects=(OldName="/Script/MagicLeapHandTracking.MagicLeapHandTrackingFunctionLibrary.GetGestureKeypoints",NewName="/Script/MagicLeapHandTracking.MagicLeapHandTrackingFunctionLibrary.GetGestureKeypoints")
++FunctionRedirects=(OldName="/Script/MagicLeapHandTracking.MagicLeapHandTrackingFunctionLibrary.SetConfiguration",NewName="/Script/MagicLeapHandTracking.MagicLeapHandTrackingFunctionLibrary.SetConfiguration")
++FunctionRedirects=(OldName="/Script/MagicLeapHandTracking.MagicLeapHandTrackingFunctionLibrary.GetConfiguration",NewName="/Script/MagicLeapHandTracking.MagicLeapHandTrackingFunctionLibrary.GetConfiguration")
++FunctionRedirects=(OldName="/Script/MagicLeapHandTracking.MagicLeapHandTrackingFunctionLibrary.SetStaticGestureConfidenceThreshold",NewName="/Script/MagicLeapHandTracking.MagicLeapHandTrackingFunctionLibrary.SetStaticGestureConfidenceThreshold")
++FunctionRedirects=(OldName="/Script/MagicLeapHandTracking.MagicLeapHandTrackingFunctionLibrary.GetStaticGestureConfidenceThreshold",NewName="/Script/MagicLeapHandTracking.MagicLeapHandTrackingFunctionLibrary.GetStaticGestureConfidenceThreshold")
++FunctionRedirects=(OldName="/Script/MagicLeapHandTracking.MagicLeapHandTrackingFunctionLibrary.GetHandGestureConfidence",NewName="/Script/MagicLeapHandTracking.MagicLeapHandTrackingFunctionLibrary.GetCurrentGestureConfidence")
++FunctionRedirects=(OldName="/Script/MagicLeapHandTracking.MagicLeapGesturesFunctionLibrary.GetCurrentGesture",NewName="/Script/MagicLeapHandTracking.MagicLeapHandTrackingFunctionLibrary.GetCurrentGesture")
 
 ; LevelSequenceDirectorBlueprint removal
 +ClassRedirects=(OldName="LevelSequenceDirectorBlueprint",NewName="/Script/LevelSequence.LegacyLevelSequenceDirectorBlueprint")
