[CONSOLE] [SWITCH] [-] Remove SwitchRuntimeSettings config section from BaseEngine.ini. There was nothing in it anyway.

diff --git a/Engine/Config/BaseEngine.ini b/Engine/Config/BaseEngine.ini
--- a/Engine/Config/BaseEngine.ini
+++ b/Engine/Config/BaseEngine.ini
@@ -1868,9 +1868,6 @@ MLSDKPath=(Path="")
 bEnableZI=False
 bUseVulkanForZI=True
 
-[/Script/SwitchRuntimeSettings.SwitchRuntimeSettings]
-; Due to secrecy, the defaults are in Engine/Switch/BaseSwitchEngine.ini _AND_ USwitchRuntimeSettings::USwitchRuntimeSettings()
-
 [/Script/LuminRuntimeSettings.LuminRuntimeSettings]
 PackageName=com.YourCompany.[PROJECT]
 ApplicationDisplayName=
