Disabled Vulkan support on Galaxy S7 (Mali) if device running on Android version less than 8, to avoid crash in the driver

diff --git a/Engine/Config/BaseDeviceProfiles.ini b/Engine/Config/BaseDeviceProfiles.ini
--- a/Engine/Config/BaseDeviceProfiles.ini
+++ b/Engine/Config/BaseDeviceProfiles.ini
@@ -358,6 +358,7 @@ MatchProfile=(Profile="Android_PowerVR54x",Match=((SourceType=SRC_GpuFamily,Comp
 +MatchProfile=(Profile="Android_Mali_T6xx",Match=((SourceType=SRC_GpuFamily,CompareType=CMP_Regex,MatchString="^Mali\\-T6")))
 +MatchProfile=(Profile="Android_Mali_T7xx",Match=((SourceType=SRC_GpuFamily,CompareType=CMP_Regex,MatchString="^Mali\\-T7")))
 +MatchProfile=(Profile="Android_Mali_T8xx_No_Vulkan",Match=((SourceType=SRC_GpuFamily,CompareType=CMP_Regex,MatchString="^Mali\\-T8"),(SourceType=SRC_VulkanVersion,CompareType=CMP_Regex,MatchString="^1\\.0\\.([0-9]+)"),(SourceType=SRC_PreviousRegexMatch,CompareType=CMP_Less,MatchString="11")))
++MatchProfile=(Profile="Android_Mali_T8xx_No_Vulkan",Match=((SourceType=SRC_GpuFamily,CompareType=CMP_Regex,MatchString="^Mali\\-T8"),(SourceType=SRC_AndroidVersion, CompareType=CMP_Regex,MatchString="([0-9]).*"),(SourceType=SRC_PreviousRegexMatch,CompareType=CMP_Less,MatchString="8")))
 +MatchProfile=(Profile="Android_Mali_T8xx",Match=((SourceType=SRC_GpuFamily,CompareType=CMP_Regex,MatchString="^Mali\\-T8")))
 +MatchProfile=(Profile="Android_Mali_G71",Match=((SourceType=SRC_GpuFamily,CompareType=CMP_Regex,MatchString="^Mali\\-G71")))
 +MatchProfile=(Profile="Android_Mali_G72",Match=((SourceType=SRC_GpuFamily,CompareType=CMP_Regex,MatchString="^Mali\\-G72")))
