[NO NON-PII COMMIT MESSAGE]

diff --git a/Engine/Config/BaseEditorPerProjectUserSettings.ini b/Engine/Config/BaseEditorPerProjectUserSettings.ini
--- a/Engine/Config/BaseEditorPerProjectUserSettings.ini
+++ b/Engine/Config/BaseEditorPerProjectUserSettings.ini
@@ -245,6 +245,9 @@ LastExecutedLaunchPlatform=
 +PhoneScreenResolutions=(Description="Apple iPhone 8",Width=375,Height=667,AspectRatio="16:9",bCanSwapAspectRatio=true,ProfileName="iPhone8")
 +PhoneScreenResolutions=(Description="Apple iPhone 8+",Width=414,Height=736,AspectRatio="16:9",bCanSwapAspectRatio=true,ProfileName="iPhone8Plus")
 +PhoneScreenResolutions=(Description="Apple iPhone X",Width=375,Height=812,AspectRatio="19.5:9",bCanSwapAspectRatio=true,ProfileName="iPhoneX")
++PhoneScreenResolutions=(Description="Apple iPhone XS",Width=375,Height=812,AspectRatio="16:9",bCanSwapAspectRatio=true,ProfileName="iPhoneXS")
++PhoneScreenResolutions=(Description="Apple iPhone XS Max",Width=414,Height=896,AspectRatio="16:9",bCanSwapAspectRatio=true,ProfileName="iPhoneXSMax")
++PhoneScreenResolutions=(Description="Apple iPhone XR",Width=414,Height=896,AspectRatio="16:9",bCanSwapAspectRatio=true,ProfileName="iPhoneXR")
 +PhoneScreenResolutions=(Description="HTC One",Width=1080,Height=1920,AspectRatio="16:9",bCanSwapAspectRatio=true,ProfileName="Android_High")
 +PhoneScreenResolutions=(Description="Samsung Galaxy S4",Width=1080,Height=1920,AspectRatio="16:9",bCanSwapAspectRatio=true,ProfileName="Android_Low")
 +PhoneScreenResolutions=(Description="Samsung Galaxy S6",Width=1440,Height=2560,AspectRatio="16:9",bCanSwapAspectRatio=true,ProfileName="Android_Mali_T7xx")
@@ -253,10 +256,13 @@ LastExecutedLaunchPlatform=
 +PhoneScreenResolutions=(Description="Samsung Galaxy S8 (Adreno)",Width=1080,Height=2220,AspectRatio="18.5:9",bCanSwapAspectRatio=true,ProfileName="Android_Adreno5xx")
 +PhoneScreenResolutions=(Description="Samsung Galaxy S9 (Mali)",Width=1440,Height=2960,AspectRatio="18.5:9",bCanSwapAspectRatio=true,ProfileName="Android_High")
 +PhoneScreenResolutions=(Description="Samsung Galaxy S9 (Adreno)",Width=1440,Height=2960,AspectRatio="18.5:9",bCanSwapAspectRatio=true,ProfileName="Android_High")
++PhoneScreenResolutions=(Description="Samsung Galaxy Note 9 (Mali)",Width=1440,Height=2960,AspectRatio="18.5:9",bCanSwapAspectRatio=true,ProfileName="Android_High")
 +PhoneScreenResolutions=(Description="Google Pixel",Width=1080,Height=1920,AspectRatio="16:9",bCanSwapAspectRatio=true,ProfileName="Android_Mid")
 +PhoneScreenResolutions=(Description="Google Pixel XL",Width=1080,Height=1920,AspectRatio="16:9",bCanSwapAspectRatio=true,ProfileName="Android_Mid")
 +PhoneScreenResolutions=(Description="Google Pixel 2",Width=1080,Height=1920,AspectRatio="16:9",bCanSwapAspectRatio=true,ProfileName="Android_Mid")
 +PhoneScreenResolutions=(Description="Google Pixel 2 XL",Width=1080,Height=1920,AspectRatio="16:9",bCanSwapAspectRatio=true,ProfileName="Android_Mid")
++PhoneScreenResolutions=(Description="Google Pixel 3",Width=1080,Height=2160,AspectRatio="18:9",bCanSwapAspectRatio=true,ProfileName="Android_Mid")
++PhoneScreenResolutions=(Description="Google Pixel 3 XL",Width=1440,Height=2960,AspectRatio="18.5:9",bCanSwapAspectRatio=true,ProfileName="Android_Mid")
 +PhoneScreenResolutions=(Description="Razer Phone",Width=1080,Height=1920,AspectRatio="16:9",bCanSwapAspectRatio=true,ProfileName="Android_Mid")
 
 ; common screen resolutions for tablet devices
