Changed the default setting for Mac scroll gesture direction from "Standard" to "Use System Settings"

#jria UE-64760

diff --git a/Engine/Config/BaseEditorPerProjectUserSettings.ini b/Engine/Config/BaseEditorPerProjectUserSettings.ini
--- a/Engine/Config/BaseEditorPerProjectUserSettings.ini
+++ b/Engine/Config/BaseEditorPerProjectUserSettings.ini
@@ -304,8 +304,8 @@ bUseAbsoluteTranslation=True
 bLevelStreamingVolumePrevis=False
 bUseUE3OrbitControls=False
 ;Scroll gesture direction
-ScrollGestureDirectionFor3DViewports=Standard
-ScrollGestureDirectionForOrthoViewports=Standard
+ScrollGestureDirectionFor3DViewports=UseSystemSetting
+ScrollGestureDirectionForOrthoViewports=UseSystemSetting
 
 bUsePowerOf2SnapSize=False
 
@@ -822,4 +822,4 @@ IncludeWindCheck=false
 IncludeSmoothLODCheck=false
 
 [/Script/MeshPaint.MeshPaintSettings]
-VertexPreviewSize=6
\ No newline at end of file
+VertexPreviewSize=6
