PR #5257: Make AddChildToWrapBox function name consistent (Contributed by [REDACTED_PERSON])

diff --git a/Engine/Config/BaseEngine.ini b/Engine/Config/BaseEngine.ini
--- a/Engine/Config/BaseEngine.ini
+++ b/Engine/Config/BaseEngine.ini
@@ -623,6 +623,9 @@ bUseStreamingPause=false
 +ClassRedirects=(OldName="LevelSequenceDirectorBlueprint",NewName="/Script/LevelSequence.LegacyLevelSequenceDirectorBlueprint")
 +ClassRedirects=(OldName="LevelSequenceDirectorGeneratedClass",NewName="/Script/Engine.BlueprintGeneratedClass")
 
+; 4.22
+FunctionRedirects=(OldName="AddChildWrapBox", NewName="AddChildToWrapBox")
+
 [CoreUObject.Metadata]
 ; Note: UnrealHeaderTool should be rerun after making changes to MetadataRedirects to catch any keys specified in class headers
 +MetadataRedirects=(OldKey="K2Protected", NewKey="BlueprintProtected")
diff --git a/Engine/Source/Runtime/UMG/Private/Components/WrapBox.cpp b/Engine/Source/Runtime/UMG/Private/Components/WrapBox.cpp
--- a/Engine/Source/Runtime/UMG/Private/Components/WrapBox.cpp
+++ b/Engine/Source/Runtime/UMG/Private/Components/WrapBox.cpp
@@ -55,6 +55,11 @@ void UWrapBox::OnSlotRemoved(UPanelSlot* InSlot)
 }
 
 UWrapBoxSlot* UWrapBox::AddChildWrapBox(UWidget* Content)
+{
+	return AddChildToWrapBox(Content);
+}
+
+UWrapBoxSlot* UWrapBox::AddChildToWrapBox(UWidget* Content)
 {
 	return Cast<UWrapBoxSlot>(Super::AddChild(Content));
 }
diff --git a/Engine/Source/Runtime/UMG/Public/Components/WrapBox.h b/Engine/Source/Runtime/UMG/Public/Components/WrapBox.h
--- a/Engine/Source/Runtime/UMG/Public/Components/WrapBox.h
+++ b/Engine/Source/Runtime/UMG/Public/Components/WrapBox.h
@@ -43,10 +43,12 @@ public:
 
 public:
 
-	/**  */
-	UFUNCTION(BlueprintCallable, Category="Panel")
+	DEPRECATED(4.22, "Deprecated, please use AddChildToWrapBox() instead")
 	UWrapBoxSlot* AddChildWrapBox(UWidget* Content);
 
+	UFUNCTION(BlueprintCallable, Category="Panel")
+	UWrapBoxSlot* AddChildToWrapBox(UWidget* Content);
+
 #if WITH_EDITOR
 	// UWidget interface
 	virtual const FText GetPaletteCategory() override;
