PR #5292: Add missing defaults for mouse invert settings (Contributed by [REDACTED_PERSON])

diff --git a/Engine/Config/BaseEditorPerProjectUserSettings.ini b/Engine/Config/BaseEditorPerProjectUserSettings.ini
--- a/Engine/Config/BaseEditorPerProjectUserSettings.ini
+++ b/Engine/Config/BaseEditorPerProjectUserSettings.ini
@@ -309,6 +309,9 @@ CameraSpeedScalar=1.0f
 ;Mouse speed when using middle mouse scroll in viewport
 MouseScrollCameraSpeed=5
 MouseSensitivty=.2f
+bInvertMouseLookYAxis=False
+bInvertOrbitYAxis=False
+bInvertMiddleMousePan=False
 bInvertRightMouseDollyYAxis=False
 ; Whether to use mouse position as direct widget position
 bUseAbsoluteTranslation=True
