Fixed bad redirector that was overwriting the functionredirector array

diff --git a/Engine/Config/BaseEngine.ini b/Engine/Config/BaseEngine.ini
--- a/Engine/Config/BaseEngine.ini
+++ b/Engine/Config/BaseEngine.ini
@@ -632,7 +632,7 @@ bUseStreamingPause=false
 
 ; 4.22
 +FunctionRedirects=(OldName="UserWidget.PlayAnimation",NewName="UserWidget.PlayAnimationAtTime")
-FunctionRedirects=(OldName="AddChildWrapBox", NewName="AddChildToWrapBox")
++FunctionRedirects=(OldName="AddChildWrapBox", NewName="AddChildToWrapBox")
 
 [CoreUObject.Metadata]
 ; Note: UnrealHeaderTool should be rerun after making changes to MetadataRedirects to catch any keys specified in class headers
