increase object size for cooking

diff --git a/Engine/Config/BaseEngine.ini b/Engine/Config/BaseEngine.ini
--- a/Engine/Config/BaseEngine.ini
+++ b/Engine/Config/BaseEngine.ini
@@ -765,7 +765,7 @@ gc.NumRetriesBeforeForcingGC=10
 gc.AllowParallelGC=True
 ; pick a fractional number to keep phase shifting and avoid collisions
 gc.TimeBetweenPurgingPendingKillObjects=61.1
-gc.MaxObjectsInEditor=8388607
+gc.MaxObjectsInEditor=16777216
 gc.IncrementalBeginDestroyEnabled=True
 gc.CreateGCClusters=True
 gc.MergeGCClusters=False
