ProfileGPUHitches should autocomplete in console

diff --git a/Engine/Config/BaseInput.ini b/Engine/Config/BaseInput.ini
--- a/Engine/Config/BaseInput.ini
+++ b/Engine/Config/BaseInput.ini
@@ -188,6 +188,7 @@ MaxScrollbackSize=1024
 +ManualAutoCompleteList=(Command="StopRecordingSequence",Desc="Stop recording the current sequence. Only one sequence recording can be active at one time.")
 +ManualAutoCompleteList=(Command="FreezeRendering",Desc="Toggle freezing of most aspects of rendering (such as visibility calculations), useful in conjunction with ToggleDebugCamera to fly around and see how frustum and occlusion culling is working")
 +ManualAutoCompleteList=(Command="ProfileGPU",Desc="Profile one frame of rendering commands sent to the GPU")
++ManualAutoCompleteList=(Command="ProfileGPUHitches",Desc="Toggle profiling of GPU hitches.")
 +ManualAutoCompleteList=(Command="Automation",Desc="Run an automation command (e.g., Automation RunTests TestName)")
 +ManualAutoCompleteList=(Command="CsvProfile Start",Desc="Start CSV profiling.")
 +ManualAutoCompleteList=(Command="CsvProfile Stop",Desc="Stop CSV profiling.")
