Sequencer: Added command auto complete list for record take and stop recording take

diff --git a/Engine/Config/BaseInput.ini b/Engine/Config/BaseInput.ini
--- a/Engine/Config/BaseInput.ini
+++ b/Engine/Config/BaseInput.ini
@@ -186,6 +186,8 @@ MaxScrollbackSize=1024
 +ManualAutoCompleteList=(Command="StopRecordingAnimation All",Desc="Stop all recording animations")
 +ManualAutoCompleteList=(Command="RecordSequence Filter ActorOrClassName",Desc="Record a level sequence from gameplay. Filter=<All|Actor|Class>")
 +ManualAutoCompleteList=(Command="StopRecordingSequence",Desc="Stop recording the current sequence. Only one sequence recording can be active at one time.")
++ManualAutoCompleteList=(Command="RecordTake Filter ActorOrClassName",Desc="Record a level sequence from gameplay. Filter=<All|Actor|Class>")
++ManualAutoCompleteList=(Command="StopRecordingTake",Desc="Stop recording the current sequence. Only one sequence recording can be active at one time.")
 +ManualAutoCompleteList=(Command="FreezeRendering",Desc="Toggle freezing of most aspects of rendering (such as visibility calculations), useful in conjunction with ToggleDebugCamera to fly around and see how frustum and occlusion culling is working")
 +ManualAutoCompleteList=(Command="ProfileGPU",Desc="Profile one frame of rendering commands sent to the GPU")
 +ManualAutoCompleteList=(Command="Automation",Desc="Run an automation command (e.g., Automation RunTests TestName)")
