Duplicating from Dev-Partner-MagicLeap-4.20 cl 4189995
HandTracking 'failed to load' errors.
-Needed a package redirector as well as the various class and enum redirectors.
#MLQDR

Files:
//UE4/Dev-Partner-MagicLeap-4.20/Engine/Config/BaseEngine.ini#8

diff --git a/Engine/Config/BaseEngine.ini b/Engine/Config/BaseEngine.ini
--- a/Engine/Config/BaseEngine.ini
+++ b/Engine/Config/BaseEngine.ini
@@ -586,6 +586,7 @@ bUseStreamingPause=false
 ; 4.21
 
 ; MagicLeap Gestures->HandTracking
++PackageRedirects=(OldName="/Script/MagicLeapGestures", NewName="/Script/MagicLeapHandTracking", MatchSubstring=true)
 +ClassRedirects=(OldName="MagicLeapGesturesFunctionLibrary",NewName="/Script/MagicLeapHandTracking.MagicLeapHandTrackingFunctionLibrary")
 +EnumRedirects=(OldName="EStaticGestures",NewName="/Script/MagicLeapHandTracking.EHandTrackingGesture")
 +EnumRedirects=(OldName="EGestureKeypointsFilterLevel",NewName="/Script/MagicLeapHandTracking.EHandTrackingKeypointFilterLevel")
