Fixes wrong cvar name for render graph immediate execution in ConsoleVariables.ini

diff --git a/Engine/Config/ConsoleVariables.ini b/Engine/Config/ConsoleVariables.ini
--- a/Engine/Config/ConsoleVariables.ini
+++ b/Engine/Config/ConsoleVariables.ini
@@ -44,7 +44,7 @@
 ; Uncomment to enable parallel rendering at startup
 ;r.RHICmdBypass=0
 ; Uncomment to get render graph executing passes as they get created to easily debug crashes caused by pass wiring logic.
-;r.Graph.ImmediateMode=1
+;r.RDG.ImmediateMode=1
 
 ; r.XGEShaderCompile is now enabled by default in source. Uncomment to disable XGE shader compilation.
 ;r.XGEShaderCompile = 0
