Disabled ES3.1 support on PowerVR Rogue G6xxx devices, as they don't work properly with it
Added PowerVRHigh device profile which will match Rogue G7xxx and up.

diff --git a/Engine/Config/BaseDeviceProfiles.ini b/Engine/Config/BaseDeviceProfiles.ini
--- a/Engine/Config/BaseDeviceProfiles.ini
+++ b/Engine/Config/BaseDeviceProfiles.ini
@@ -68,6 +68,8 @@
 +DeviceProfileNameAndTypes=Android_Adreno5xx,Android
 +DeviceProfileNameAndTypes=Android_Adreno6xx,Android
 +DeviceProfileNameAndTypes=Android_PowerVR54x,Android
++DeviceProfileNameAndTypes=Android_PowerVR_G6xxx,Android
++DeviceProfileNameAndTypes=Android_PowerVRHigh,Android
 +DeviceProfileNameAndTypes=Android_PowerVRLow,Android
 +DeviceProfileNameAndTypes=Android_PowerVRMid,Android
 +DeviceProfileNameAndTypes=Android_TegraK1,Android
@@ -371,6 +373,8 @@ BaseProfileName=IPadPro
 
 [/Script/AndroidDeviceProfileSelector.AndroidDeviceProfileMatchingRules]
 MatchProfile=(Profile="Android_PowerVR54x",Match=((SourceType=SRC_GpuFamily,CompareType=CMP_Regex,MatchString="PowerVR SGX 54[0-9]")))
++MatchProfile=(Profile="Android_PowerVR_G6xxx",Match=((SourceType=SRC_GpuFamily,CompareType=CMP_Regex,MatchString="PowerVR Rogue G6[0-9]+")))
++MatchProfile=(Profile="Android_PowerVRHigh",Match=((SourceType=SRC_GpuFamily,CompareType=CMP_Regex,MatchString="PowerVR Rogue G"),(SourceType=SRC_GlVersion,CompareType=CMP_Regex,MatchString="^OpenGL ES 3\\.")))
 +MatchProfile=(Profile="Android_PowerVRMid",Match=((SourceType=SRC_GpuFamily,CompareType=CMP_Regex,MatchString="PowerVR"),(SourceType=SRC_GlVersion,CompareType=CMP_Regex,MatchString="^OpenGL ES 3\\.")))
 +MatchProfile=(Profile="Android_PowerVRLow",Match=((SourceType=SRC_GpuFamily,CompareType=CMP_Regex,MatchString="PowerVR")))
 +MatchProfile=(Profile="Android_Adreno2xx",Match=((SourceType=SRC_GpuFamily,CompareType=CMP_Regex,MatchString="Adreno \\(TM\\) 2[0-9][0-9]")))
@@ -514,12 +518,22 @@ BaseProfileName=Android_Mid
 +CVars=r.Mobile.ForceDepthResolve=1
 +CVars=r.Android.DisableASTCSupport=1
 
+[Android_PowerVR_G6xxx DeviceProfile]
+DeviceType=Android
+BaseProfileName=Android_PowerVRMid
+; Disable ES3.1 support on PoverVR G6xxx devices
++CVars=r.Android.DisableOpenGLES31Support=1
+
 [Android_PowerVRLow DeviceProfile]
 DeviceType=Android
 BaseProfileName=Android_Low
 +CVars=r.Mobile.ForceDepthResolve=1
 +CVars=r.Android.DisableASTCSupport=1
 
+[Android_PowerVRHigh DeviceProfile]
+DeviceType=Android
+BaseProfileName=Android_High
+
 [Android_IntelHD DeviceProfile]
 DeviceType=Android
 BaseProfileName=Android_Low
