Disable Vulkan support on Adreno 5xx based devices running Android 6, driver issues.

diff --git a/Engine/Config/BaseDeviceProfiles.ini b/Engine/Config/BaseDeviceProfiles.ini
--- a/Engine/Config/BaseDeviceProfiles.ini
+++ b/Engine/Config/BaseDeviceProfiles.ini
@@ -66,6 +66,8 @@
 +DeviceProfileNameAndTypes=Android_Adreno4xx_Nexus6,Android
 +DeviceProfileNameAndTypes=Android_Adreno4xx,Android
 +DeviceProfileNameAndTypes=Android_Adreno5xx,Android
++DeviceProfileNameAndTypes=Android_Adreno5xx_No_Vulkan,Android
++DeviceProfileNameAndTypes=Android_Adreno5xx_Vulkan_Nougat,Android
 +DeviceProfileNameAndTypes=Android_Adreno6xx,Android
 +DeviceProfileNameAndTypes=Android_PowerVR54x,Android
 +DeviceProfileNameAndTypes=Android_PowerVR_G6xxx,Android
@@ -384,6 +386,7 @@ MatchProfile=(Profile="Android_PowerVR54x",Match=((SourceType=SRC_GpuFamily,Comp
 +MatchProfile=(Profile="Android_Adreno4xx_Nexus6",Match=((SourceType=SRC_DeviceModel,CompareType=CMP_Equal,MatchString="Nexus 6")))
 +MatchProfile=(Profile="Android_Adreno4xx",Match=((SourceType=SRC_GpuFamily,CompareType=CMP_Regex,MatchString="Adreno \\(TM\\) 4[0-9][0-9]")))
 +MatchProfile=(Profile="Android_Adreno51x",Match=((SourceType=SRC_GpuFamily,CompareType=CMP_Regex,MatchString="Adreno \\(TM\\) 51[0-9]")))
++MatchProfile=(Profile="Android_Adreno5xx_No_Vulkan",Match=((SourceType=SRC_GpuFamily,CompareType=CMP_Regex,MatchString="Adreno \\(TM\\) 5[0-9][0-9]"),(SourceType=SRC_AndroidVersion,CompareType=CMP_Regex,MatchString="6\\.")))
 +MatchProfile=(Profile="Android_Adreno5xx_Vulkan_Nougat",Match=((SourceType=SRC_GpuFamily,CompareType=CMP_Regex,MatchString="Adreno \\(TM\\) 5[0-9][0-9]"),,(SourceType=SRC_VulkanAvailable,CompareType=CMP_Equal,MatchString="true"),(SourceType=SRC_AndroidVersion,CompareType=CMP_Regex,MatchString="7\\.")))
 +MatchProfile=(Profile="Android_Adreno5xx",Match=((SourceType=SRC_GpuFamily,CompareType=CMP_Regex,MatchString="Adreno \\(TM\\) 5[0-9][0-9]")))
 +MatchProfile=(Profile="Android_Adreno6xx",Match=((SourceType=SRC_GpuFamily,CompareType=CMP_Regex,MatchString="Adreno \\(TM\\) 6[0-9][0-9]")))
@@ -487,6 +490,11 @@ DeviceType=Android
 BaseProfileName=Android_High
 +CVars=r.DisjointTimerQueries=1
 
+[Android_Adreno5xx_No_Vulkan DeviceProfile]
+DeviceType=Android
+BaseProfileName=Android_Adreno5xx
++CVars=r.Android.DisableVulkanSupport=1
+
 [Android_Adreno5xx_Vulkan_Nougat DeviceProfile]
 DeviceType=Android
 BaseProfileName=Android_Adreno5xx
