Fixes to spectator JIP. Make sure spectator beacon timeout variables are initialized.

diff --git a/Engine/Config/BaseEngine.ini b/Engine/Config/BaseEngine.ini
--- a/Engine/Config/BaseEngine.ini
+++ b/Engine/Config/BaseEngine.ini
@@ -1436,6 +1436,11 @@ bLogoutOnSessionTimeout=true
 SessionTimeoutSecs=10
 TravelSessionTimeoutSecs=45
 
+[/Script/OnlineSubsystemUtils.SpectatorBeaconHost]
+bLogoutOnSessionTimeout=true
+SessionTimeoutSecs=10
+TravelSessionTimeoutSecs=45
+
 [/Script/Lobby.LobbyBeaconClient]
 BeaconConnectionInitialTimeout=90.0
 BeaconConnectionTimeout=45.0
