Updated BaseEngine.ini TargetSDKVersion=26 for Android

#ue4
#android

diff --git a/Engine/Config/BaseEngine.ini b/Engine/Config/BaseEngine.ini
--- a/Engine/Config/BaseEngine.ini
+++ b/Engine/Config/BaseEngine.ini
@@ -1895,7 +1895,7 @@ PackageName=com.YourCompany.[PROJECT]
 StoreVersion=1
 VersionDisplayName=1.0
 MinSDKVersion=9
-TargetSDKVersion=9
+TargetSDKVersion=26
 bEnableGradle=true
 bEnableLint=false
 bShowLaunchImage=true
