Enabled the "Orbit camera around selection" setting by default.

diff --git a/Engine/Config/BaseEditorPerProjectUserSettings.ini b/Engine/Config/BaseEditorPerProjectUserSettings.ini
--- a/Engine/Config/BaseEditorPerProjectUserSettings.ini
+++ b/Engine/Config/BaseEditorPerProjectUserSettings.ini
@@ -318,6 +318,7 @@ bUseAbsoluteTranslation=True
 bLevelStreamingVolumePrevis=False
 bUseUE3OrbitControls=False
 bUseDistanceScaledCameraSpeed=True
+bOrbitCameraAroundSelection=True
 ;Scroll gesture direction
 ScrollGestureDirectionFor3DViewports=UseSystemSetting
 ScrollGestureDirectionForOrthoViewports=UseSystemSetting
