Changed default mapping of speaker angles  to SideLeft=250 and SideRight=110 to be better inline with Dolby spec when these angles are used for both 7.1 and 5.1 configurations.

diff --git a/Engine/Config/BaseEngine.ini b/Engine/Config/BaseEngine.ini
--- a/Engine/Config/BaseEngine.ini
+++ b/Engine/Config/BaseEngine.ini
@@ -822,7 +822,7 @@ FrontLeftOfCenter=15
 FrontRightOfCenter=345
 BackCenter=180
 SideLeft=250
-SideRight=100
+SideRight=110
 
 ; These are the default channel orderings to use when using pro-audio interfaces that don't report channel order
 [AudioDefaultChannelOrder]
