Centered notch unsafe zone in BP widget templates for iPhone XS Max and XR #Jira UE-66306 #rb trivial

diff --git a/Engine/Config/BaseDeviceProfiles.ini b/Engine/Config/BaseDeviceProfiles.ini
--- a/Engine/Config/BaseDeviceProfiles.ini
+++ b/Engine/Config/BaseDeviceProfiles.ini
@@ -309,12 +309,12 @@ BaseProfileName=iPhone8
 [iPhoneXSMax DeviceProfile]
 DeviceType=IOS
 BaseProfileName=iPhone8
-+CVars=r.CustomUnsafeZones="(L:free[0,-15][812,15]);(P:fixed[83,0][206,30])"
++CVars=r.CustomUnsafeZones="(L:free[0,-15][812,15]);(P:fixed[104,0][206,30])"
 
 [iPhoneXR DeviceProfile]
 DeviceType=IOS
 BaseProfileName=iPhone8
-+CVars=r.CustomUnsafeZones="(L:free[0,-15][812,15]);(P:fixed[83,0][206,30])"
++CVars=r.CustomUnsafeZones="(L:free[0,-15][812,15]);(P:fixed[104,0][206,30])"
 
 [iPadPro DeviceProfile]
 DeviceType=IOS
