Make sure all fbx import value are in the config file, this allow the reset to default to work well.

diff --git a/Engine/Config/BaseEditorPerProjectUserSettings.ini b/Engine/Config/BaseEditorPerProjectUserSettings.ini
--- a/Engine/Config/BaseEditorPerProjectUserSettings.ini
+++ b/Engine/Config/BaseEditorPerProjectUserSettings.ini
@@ -610,46 +610,89 @@ BackupIntervalInMinutes=5
 
 
 [/Script/UnrealEd.FbxImportUI]
-bImportMaterials=True
-bAutoCreateGroups=True
-bInvertNormalMaps=False
-bImportTextures=True
 bOverrideFullName=True
-bConvertScene=True
-bForceFrontXAxis=False
-bConvertSceneUnit=False
-bRemoveNamespaces=True
-bImportAnimations=False
-bResampleAnimations=False
-bImportRigidMesh=False
-bCombineMeshes=True
 bCreatePhysicsAsset=True
-bImportMesh=True
+bAutoComputeLodDistances=True
+LodDistance0=0.0
+LodDistance1=0.0
+LodDistance2=0.0
+LodDistance3=0.0
+LodDistance4=0.0
+LodDistance5=0.0
+LodDistance6=0.0
+LodDistance7=0.0
 MinimumLodNumber=0
 LodNumber=0
-bAutoComputeLodDistances=True
+bImportAnimations=False
+bImportMaterials=True
+bImportTextures=True
 
+[/Script/UnrealEd.FbxAssetImportData]
+ImportTranslation=(X=0.0,Y=0.0,Z=0.0)
+ImportRotation=(Pitch=0.0,Yaw=0.0,Roll=0.0)
+ImportUniformScale=1.0
+bConvertScene=True
+bForceFrontXAxis=False
+bConvertSceneUnit=False
 
+[/Script/UnrealEd.FbxMeshImportData]
+bTransformVertexToAbsolute=True
+bBakePivotInVertex=False
+bReorderMaterialToFbxOrder=True
+bImportMeshLODs=False
+NormalImportMethod=FBXNIM_ComputeNormals
+NormalGenerationMethod=EFBXNormalGenerationMethod::MikkTSpace
 
 [/Script/UnrealEd.FbxStaticMeshImportData]
-bOneConvexHullPerUCX=True
-bImportMeshLODs=False
-NormalImportMethod=FBXNIM_ImportNormals
+StaticMeshLODGroup=""
 VertexColorImportOption=EVertexColorImportOption::Ignore
 VertexOverrideColor=(R=255,G=255,B=255,A=255)
-
+bRemoveDegenerates=True
+bBuildAdjacencyBuffer=True
+bBuildReversedIndexBuffer=True
+bGenerateLightmapUVs=True
+bOneConvexHullPerUCX=True
+bAutoGenerateCollision=True
+bCombineMeshes=False
+; Override of the base class default value
+NormalImportMethod=FBXNIM_ImportNormals
 
 [/Script/UnrealEd.FbxSkeletalMeshImportData]
+VertexColorImportOption=EVertexColorImportOption::Replace
+VertexOverrideColor=(R=0,G=0,B=0,A=0)
+bUpdateSkeletonReferencePose=False
+bUseT0AsRefPose=False
 bPreserveSmoothingGroups=True
-bImportMeshLODs=False
+bImportMeshesInBoneHierarchy=True
 bImportMorphTargets=False
 ThresholdPosition=0.00002
 ThresholdTangentNormal=0.00002
 ThresholdUV=0.0009765625
 
+[/Script/UnrealEd.FbxAnimSequenceImportData]
+bImportMeshesInBoneHierarchy=True
+AnimationLength=FBXALIT_ExportedTime
+FrameImportRange=(Min=0, Max=0)
+bUseDefaultSampleRate=False
+CustomSampleRate=0
+bImportCustomAttribute=True
+bImportBoneTracks=True
+bSetMaterialDriveParameterOnCustomAttribute=False
+bRemoveRedundantKeys=True
+bDeleteExistingMorphTargetCurves=False
+bDoNotImportCurveWithZero=True
+bPreserveLocalTransform=False
 
 [/Script/UnrealEd.FbxTextureImportData]
-
+bInvertNormalMaps=False
+MaterialSearchLocation=EMaterialSearchLocation::Local
+BaseMaterialName=""
+BaseColorName=""
+BaseDiffuseTextureName=""
+BaseNormalTextureName=""
+BaseEmissiveColorName=""
+BaseEmmisiveTextureName=""
+BaseSpecularTextureName=""
 
 [SoundSettings]
 ChirpSoundClasses=Dialog DialogMedium DialogLoud DialogDeafening
diff --git a/Engine/Source/Editor/UnrealEd/Classes/Factories/FbxAnimSequenceImportData.h b/Engine/Source/Editor/UnrealEd/Classes/Factories/FbxAnimSequenceImportData.h
--- a/Engine/Source/Editor/UnrealEd/Classes/Factories/FbxAnimSequenceImportData.h
+++ b/Engine/Source/Editor/UnrealEd/Classes/Factories/FbxAnimSequenceImportData.h
@@ -39,7 +39,7 @@ class UNREALED_API UFbxAnimSequenceImportData : public UFbxAssetImportData
 	GENERATED_UCLASS_BODY()
 	
 	/** If checked, meshes nested in bone hierarchies will be imported instead of being converted to bones. */
-	UPROPERTY(EditAnywhere, AdvancedDisplay, Category = ImportSettings, meta = (ImportType = "Animation"))
+	UPROPERTY(EditAnywhere, AdvancedDisplay, config, Category = ImportSettings, meta = (ImportType = "Animation"))
 	bool bImportMeshesInBoneHierarchy;
 	
 	/** Which animation range to import. The one defined at Exported, at Animated time or define a range manually */
@@ -55,7 +55,7 @@ class UNREALED_API UFbxAnimSequenceImportData : public UFbxAssetImportData
 	int32	EndFrame_DEPRECATED;
 
 	/** Frame range used when Set Range is used in Animation Length */
-	UPROPERTY(EditAnywhere, AdvancedDisplay, Category = ImportSettings, meta=(UIMin=0, ClampMin=0))
+	UPROPERTY(EditAnywhere, AdvancedDisplay, config, Category = ImportSettings, meta=(UIMin=0, ClampMin=0))
 	FInt32Interval FrameImportRange;
 
 	/** Enable this option to use default sample rate for the imported animation at 30 frames per second */
diff --git a/Engine/Source/Editor/UnrealEd/Classes/Factories/FbxSkeletalMeshImportData.h b/Engine/Source/Editor/UnrealEd/Classes/Factories/FbxSkeletalMeshImportData.h
--- a/Engine/Source/Editor/UnrealEd/Classes/Factories/FbxSkeletalMeshImportData.h
+++ b/Engine/Source/Editor/UnrealEd/Classes/Factories/FbxSkeletalMeshImportData.h
@@ -66,7 +66,7 @@ public:
 	FColor VertexOverrideColor;
 
 	/** Enable this option to update Skeleton (of the mesh)'s reference pose. Mesh's reference pose is always updated.  */
-	UPROPERTY(EditAnywhere, AdvancedDisplay, Category=Mesh, meta=(ImportType="SkeletalMesh|RigOnly", ToolTip="If enabled, update the Skeleton (of the mesh being imported)'s reference pose."))
+	UPROPERTY(EditAnywhere, AdvancedDisplay, config, Category=Mesh, meta=(ImportType="SkeletalMesh|RigOnly", ToolTip="If enabled, update the Skeleton (of the mesh being imported)'s reference pose."))
 	uint32 bUpdateSkeletonReferencePose:1;
 
 	/** Enable this option to use frame 0 as reference pose */
diff --git a/Engine/Source/Editor/UnrealEd/Classes/Factories/FbxStaticMeshImportData.h b/Engine/Source/Editor/UnrealEd/Classes/Factories/FbxStaticMeshImportData.h
--- a/Engine/Source/Editor/UnrealEd/Classes/Factories/FbxStaticMeshImportData.h
+++ b/Engine/Source/Editor/UnrealEd/Classes/Factories/FbxStaticMeshImportData.h
@@ -31,14 +31,14 @@ class UFbxStaticMeshImportData : public UFbxMeshImportData
 	FColor VertexOverrideColor;
 
 	/** Disabling this option will keep degenerate triangles found.  In general you should leave this option on. */
-	UPROPERTY(EditAnywhere, BlueprintReadWrite, AdvancedDisplay, Category = Mesh, meta = (ImportType = "StaticMesh", ReimportRestrict = "true"))
+	UPROPERTY(EditAnywhere, BlueprintReadWrite, AdvancedDisplay, config, Category = Mesh, meta = (ImportType = "StaticMesh", ReimportRestrict = "true"))
 	uint32 bRemoveDegenerates:1;
 
 	/** Required for PNT tessellation but can be slow. Recommend disabling for larger meshes. */
-	UPROPERTY(EditAnywhere, BlueprintReadWrite, AdvancedDisplay, Category = Mesh, meta = (ImportType = "StaticMesh", ReimportRestrict = "true"))
+	UPROPERTY(EditAnywhere, BlueprintReadWrite, AdvancedDisplay, config, Category = Mesh, meta = (ImportType = "StaticMesh", ReimportRestrict = "true"))
 	uint32 bBuildAdjacencyBuffer:1;
 
-	UPROPERTY(EditAnywhere, BlueprintReadWrite, AdvancedDisplay, Category = Mesh, meta = (ImportType = "StaticMesh", ReimportRestrict = "true"))
+	UPROPERTY(EditAnywhere, BlueprintReadWrite, AdvancedDisplay, config, Category = Mesh, meta = (ImportType = "StaticMesh", ReimportRestrict = "true"))
 	uint32 bBuildReversedIndexBuffer:1;
 
 	UPROPERTY(EditAnywhere, BlueprintReadWrite, config, AdvancedDisplay, Category= Mesh, meta=(ImportType="StaticMesh", ReimportRestrict = "true"))
diff --git a/Engine/Source/Editor/UnrealEd/Private/Fbx/FbxFactory.cpp b/Engine/Source/Editor/UnrealEd/Private/Fbx/FbxFactory.cpp
--- a/Engine/Source/Editor/UnrealEd/Private/Fbx/FbxFactory.cpp
+++ b/Engine/Source/Editor/UnrealEd/Private/Fbx/FbxFactory.cpp
@@ -1059,6 +1059,9 @@ void UFbxImportUI::ResetToDefault()
 	StaticMeshImportData->ReloadConfig();
 	SkeletalMeshImportData->ReloadConfig();
 	TextureImportData->ReloadConfig();
+	
+	//Make sure the UI do not display the Base Material
+	TextureImportData->bUseBaseMaterial = false;
 }
 
 namespace ImportCompareHelper
