Fix for missing mobile tutorial icon

#rnx

diff --git a/Engine/Config/BaseEditorSettings.ini b/Engine/Config/BaseEditorSettings.ini
--- a/Engine/Config/BaseEditorSettings.ini
+++ b/Engine/Config/BaseEditorSettings.ini
@@ -53,7 +53,7 @@ IncludeCollectionNames=True
 +Categories=(Identifier="Animation",Title=NSLOCTEXT("TutorialCategories","AnimationTitle","Animation"),Description=NSLOCTEXT("TutorialCategories","AnimationDescription","Tutorials covering the animation system in Unreal Engine 4."),Icon=,Texture=/Engine/Tutorial/SubEditors/TutorialAssets/icon_ShowSkeletalMeshes_40x.icon_ShowSkeletalMeshes_40x,SortOrder=600)
 +Categories=(Identifier="Landscape",Title=NSLOCTEXT("TutorialCategories","LandscapeTitle","Landscape"),Description=NSLOCTEXT("TutorialCategories","LandscapeDescription","Tutorials covering the Unreal Editor 4 terrain editor: Landscape."),Icon="LevelEditor.LandscapeMode",Texture=/Engine/Tutorial/Landscape/TutorialAssets/Landscape.Landscape,SortOrder=700)
 +Categories=(Identifier="Foliage",Title=NSLOCTEXT("TutorialCategories","FoliageTitle","Foliage"),Description=NSLOCTEXT("TutorialCategories","FoliageDescription","Tutorials covering the Unreal Engine 4 Foliage tool."),Icon="LevelEditor.FoliageMode",Texture=/Engine/Tutorial/Foliage/TutorialAssets/Foliage.Foliage,SortOrder=800)
-+Categories=(Identifier="Mobile",Title=NSLOCTEXT("TutorialCategories","MobileTitle","Mobile"),Description=NSLOCTEXT("TutorialCategories","MobileDescription","Mobile Tutorials."),Icon="MaterialEditor.ToggleMobileStats",Texture=None,SortOrder=900)
++Categories=(Identifier="Mobile",Title=NSLOCTEXT("TutorialCategories","MobileTitle","Mobile"),Description=NSLOCTEXT("TutorialCategories","MobileDescription","Mobile Tutorials."),Icon="MaterialEditor.TogglePlatformStats",Texture=None,SortOrder=900)
 +StartupTutorial=/Engine/Tutorial/Basics/LevelEditorAttract.LevelEditorAttract_C
 +TutorialContexts=(Context="StaticMeshEditor",BrowserFilter=,AttractTutorial=None,LaunchTutorial=/Engine/Tutorial/SubEditors/StaticMeshEditorTutorial.StaticMeshEditorTutorial_C)
 +TutorialContexts=(Context="LevelEditor",BrowserFilter=,AttractTutorial=None,LaunchTutorial=/Engine/Tutorial/Basics/LevelEditorOverview.LevelEditorOverview_C)
