Adds r.RDG.EmitWarnings to ConsoleVariables.ini for rendering programmer convenience.

diff --git a/Engine/Config/ConsoleVariables.ini b/Engine/Config/ConsoleVariables.ini
--- a/Engine/Config/ConsoleVariables.ini
+++ b/Engine/Config/ConsoleVariables.ini
@@ -43,8 +43,11 @@
 ;r.DumpShaderDebugWorkerCommandLine=1
 ; Uncomment to enable parallel rendering at startup
 ;r.RHICmdBypass=0
+
 ; Uncomment to get render graph executing passes as they get created to easily debug crashes caused by pass wiring logic.
 ;r.RDG.ImmediateMode=1
+; Uncomment to have render graph emitting warnings about inefficiencies.
+;r.RDG.EmitWarnings=1
 
 ; r.XGEShaderCompile is now enabled by default in source. Uncomment to disable XGE shader compilation.
 ;r.XGEShaderCompile = 0
