Removed 'stat hicthes' from autocompletion list as it no longer exists in the engine

diff --git a/Engine/Config/BaseInput.ini b/Engine/Config/BaseInput.ini
--- a/Engine/Config/BaseInput.ini
+++ b/Engine/Config/BaseInput.ini
@@ -97,7 +97,6 @@ MaxScrollbackSize=1024
 +ManualAutoCompleteList=(Command="Stat NamedEvents",Desc="Stat NamedEvents (Enables named events for external profilers)")
 +ManualAutoCompleteList=(Command="Stat StartFile",Desc="Stat StartFile (starts a stats capture, creating a new file in the Profiling directory; stop with stat StopFile to close the file)")
 +ManualAutoCompleteList=(Command="Stat StopFile",Desc="Stat StopFile (finishes a stats capture started by stat StartFile)")
-+ManualAutoCompleteList=(Command="Stat Hitches",Desc="Stat Hitches (shows an onscreen history of hitches)")
 +ManualAutoCompleteList=(Command="Stat CPULoad",Desc="Stat CPULoad (Shows CPU Utilization)")
 +ManualAutoCompleteList=(Command="Stat DUMPHITCHES",Desc="executes dumpstats on hitches - see log")
 +ManualAutoCompleteList=(Command="Stat D3D11RHI",Desc="Shows Direct3D 11 stats")
