Keep the deprecated shader platform but comment it out

diff --git a/Engine/Config/BaseEngine.ini b/Engine/Config/BaseEngine.ini
--- a/Engine/Config/BaseEngine.ini
+++ b/Engine/Config/BaseEngine.ini
@@ -2083,6 +2083,10 @@ TargetArchitecture=X86_64UnknownLinuxGnu
 
 ; the engine will default to the first targeted RHI.
 +TargetedRHIs=SF_VULKAN_SM5
+
+; OpenGL4 is deprecated, you can comment this back in to add it to your targeted RHI list
+; +TargetedRHIs=GLSL_430
+
  
 [/Script/MacTargetPlatform.MacTargetSettings]
 MaxShaderLanguageVersion=3
