Bumping the max Uobject count up in the editor by 8M

diff --git a/Engine/Config/BaseEngine.ini b/Engine/Config/BaseEngine.ini
--- a/Engine/Config/BaseEngine.ini
+++ b/Engine/Config/BaseEngine.ini
@@ -823,7 +823,7 @@ gc.NumRetriesBeforeForcingGC=10
 gc.AllowParallelGC=True
 ; pick a fractional number to keep phase shifting and avoid collisions
 gc.TimeBetweenPurgingPendingKillObjects=61.1
-gc.MaxObjectsInEditor=16777216
+gc.MaxObjectsInEditor=25165824
 gc.IncrementalBeginDestroyEnabled=True
 gc.CreateGCClusters=True
 gc.MinGCClusterSize=5
