Add Samsung S10 and S10e to mobile phones resolutions

#ue4
#android

diff --git a/Engine/Config/BaseEditorPerProjectUserSettings.ini b/Engine/Config/BaseEditorPerProjectUserSettings.ini
--- a/Engine/Config/BaseEditorPerProjectUserSettings.ini
+++ b/Engine/Config/BaseEditorPerProjectUserSettings.ini
@@ -255,6 +255,10 @@ LastExecutedLaunchPlatform=
 +PhoneScreenResolutions=(Description="Samsung Galaxy S9 (Mali)",Width=1440,Height=2960,AspectRatio="18.5:9",bCanSwapAspectRatio=true,ProfileName="Android_High")
 +PhoneScreenResolutions=(Description="Samsung Galaxy S9 (Adreno)",Width=1440,Height=2960,AspectRatio="18.5:9",bCanSwapAspectRatio=true,ProfileName="Android_High")
 +PhoneScreenResolutions=(Description="Samsung Galaxy Note 9 (Mali)",Width=1440,Height=2960,AspectRatio="18.5:9",bCanSwapAspectRatio=true,ProfileName="Android_High")
++PhoneScreenResolutions=(Description="Samsung Galaxy S10 (Adreno)",Width=1440,Height=3040,AspectRatio="19:9",bCanSwapAspectRatio=true,ProfileName="Android_Adreno6xx")
++PhoneScreenResolutions=(Description="Samsung Galaxy S10 (Mali)",Width=1440,Height=3040,AspectRatio="19:9",bCanSwapAspectRatio=true,ProfileName="Android_Mali_G76")
++PhoneScreenResolutions=(Description="Samsung Galaxy S10e (Adreno)",Width=1080,Height=2280,AspectRatio="19:9",bCanSwapAspectRatio=true,ProfileName="Android_Adreno6xx")
++PhoneScreenResolutions=(Description="Samsung Galaxy S10e (Mali)",Width=1080,Height=2280,AspectRatio="19:9",bCanSwapAspectRatio=true,ProfileName="Android_Mali_G76")
 +PhoneScreenResolutions=(Description="Google Pixel",Width=1080,Height=1920,AspectRatio="16:9",bCanSwapAspectRatio=true,ProfileName="Android_Mid")
 +PhoneScreenResolutions=(Description="Google Pixel XL",Width=1080,Height=1920,AspectRatio="16:9",bCanSwapAspectRatio=true,ProfileName="Android_Mid")
 +PhoneScreenResolutions=(Description="Google Pixel 2",Width=1080,Height=1920,AspectRatio="16:9",bCanSwapAspectRatio=true,ProfileName="Android_Mid")
