Fixed wrong scaslability variable name.

diff --git a/Engine/Config/BaseScalability.ini b/Engine/Config/BaseScalability.ini
--- a/Engine/Config/BaseScalability.ini
+++ b/Engine/Config/BaseScalability.ini
@@ -395,7 +395,7 @@ r.SkyAtmosphere.AerialPerspectiveLUT.SampleCountPerSlice=1
 r.SkyAtmosphere.AerialPerspectiveLUT.DepthResolution=8
 r.SkyAtmosphere.FastSkyLUT=1
 r.SkyAtmosphere.FastSkyLUT.SampleCountMin=2.0
-r.SkyAtmosphere.FastSkyLUT.SampleCountMin=16.0
+r.SkyAtmosphere.FastSkyLUT.SampleCountMax=16.0
 r.SkyAtmosphere.TransmittanceLUT.UseSmallFormat=1
 r.SkyAtmosphere.MultiScatteringLUT.UseSmallFormat=1
 
