added PacketSimulationProfile.Off setting for convenience in making toggleable hotkeys

#rnx

diff --git a/Engine/Config/BaseEngine.ini b/Engine/Config/BaseEngine.ini
--- a/Engine/Config/BaseEngine.ini
+++ b/Engine/Config/BaseEngine.ini
@@ -2178,6 +2178,15 @@ LLVMWasmBackend=False
 [Browser]
 bForceMessageLoop=true
 
+[PacketSimulationProfile.Off]
+; Preset that disables packet simulation
+PktLoss=0
+PktIncomingLoss=0
+PktLagMin=0
+PktLagMax=0
+PktIncomingLagMin=0
+PktIncomingLagMax=0
+
 [PacketSimulationProfile.Average]
 ; Simulates network conditions with a ping of 60 to 120
 PktLoss=1
