Missed file setting minimum iOS version to 11

diff --git a/Engine/Config/BaseEngine.ini b/Engine/Config/BaseEngine.ini
--- a/Engine/Config/BaseEngine.ini
+++ b/Engine/Config/BaseEngine.ini
@@ -1909,7 +1909,7 @@ BundleName=[PROJECT_NAME]
 BundleIdentifier=com.YourCompany.[PROJECT_NAME]
 VersionInfo=1.0
 FrameRateLock=PUFRL_30
-MinimumiOSVersion=IOS_10
+MinimumiOSVersion=IOS_11
 bSupportsIPad=True
 bSupportsIPhone=True
 AdditionalPlistData=
