Updated basescalability to haver vall variaible set for each level (did not know that other variable were not reset when swtiching level)

diff --git a/Engine/Config/BaseScalability.ini b/Engine/Config/BaseScalability.ini
--- a/Engine/Config/BaseScalability.ini
+++ b/Engine/Config/BaseScalability.ini
@@ -392,12 +392,16 @@ r.ParticleLightQuality=0
 r.DetailLevel=0
 r.SkyAtmosphere.AerialPerspectiveLUT.FastApplyOnOpaque=1 ; Always have FastSkyLUT 1 in this case to avoid wrong sky
 r.SkyAtmosphere.AerialPerspectiveLUT.SampleCountPerSlice=1
-r.SkyAtmosphere.AerialPerspectiveLUT.DepthResolution=8
+r.SkyAtmosphere.AerialPerspectiveLUT.DepthResolution=8.0
 r.SkyAtmosphere.FastSkyLUT=1
 r.SkyAtmosphere.FastSkyLUT.SampleCountMin=2.0
 r.SkyAtmosphere.FastSkyLUT.SampleCountMax=16.0
+r.SkyAtmosphere.SampleCountMin=2.0
+r.SkyAtmosphere.SampleCountMax=16.0
 r.SkyAtmosphere.TransmittanceLUT.UseSmallFormat=1
+r.SkyAtmosphere.TransmittanceLUT.SampleCount=10.0
 r.SkyAtmosphere.MultiScatteringLUT.UseSmallFormat=1
+r.SkyAtmosphere.MultiScatteringLUT.SampleCount=15.0
 
 [EffectsQuality@1]
 r.TranslucencyLightingVolumeDim=32
@@ -414,8 +418,18 @@ r.SSS.HalfRes=1
 r.EmitterSpawnRateScale=0.25
 r.ParticleLightQuality=0
 r.DetailLevel=1
-r.SkyAtmosphere.FastSkyLUT=1
 r.SkyAtmosphere.AerialPerspectiveLUT.FastApplyOnOpaque=1 ; Always have FastSkyLUT 1 in this case to avoid wrong sky
+r.SkyAtmosphere.AerialPerspectiveLUT.SampleCountPerSlice=1
+r.SkyAtmosphere.AerialPerspectiveLUT.DepthResolution=16.0
+r.SkyAtmosphere.FastSkyLUT=1
+r.SkyAtmosphere.FastSkyLUT.SampleCountMin=4.0
+r.SkyAtmosphere.FastSkyLUT.SampleCountMax=32.0
+r.SkyAtmosphere.SampleCountMin=4.0
+r.SkyAtmosphere.SampleCountMax=32.0
+r.SkyAtmosphere.TransmittanceLUT.UseSmallFormat=0
+r.SkyAtmosphere.TransmittanceLUT.SampleCount=10.0
+r.SkyAtmosphere.MultiScatteringLUT.UseSmallFormat=0
+r.SkyAtmosphere.MultiScatteringLUT.SampleCount=15.0
 
 [EffectsQuality@2]
 r.TranslucencyLightingVolumeDim=48
@@ -432,8 +446,18 @@ r.SSS.HalfRes=1
 r.EmitterSpawnRateScale=0.5
 r.ParticleLightQuality=1
 r.DetailLevel=2
-r.SkyAtmosphere.FastSkyLUT=1
 r.SkyAtmosphere.AerialPerspectiveLUT.FastApplyOnOpaque=1 ; Always have FastSkyLUT 1 in this case to avoid wrong sky
+r.SkyAtmosphere.AerialPerspectiveLUT.SampleCountPerSlice=1
+r.SkyAtmosphere.AerialPerspectiveLUT.DepthResolution=16.0
+r.SkyAtmosphere.FastSkyLUT=1
+r.SkyAtmosphere.FastSkyLUT.SampleCountMin=4.0
+r.SkyAtmosphere.FastSkyLUT.SampleCountMax=32.0
+r.SkyAtmosphere.SampleCountMin=4.0
+r.SkyAtmosphere.SampleCountMax=32.0
+r.SkyAtmosphere.TransmittanceLUT.UseSmallFormat=0
+r.SkyAtmosphere.TransmittanceLUT.SampleCount=10.0
+r.SkyAtmosphere.MultiScatteringLUT.UseSmallFormat=0
+r.SkyAtmosphere.MultiScatteringLUT.SampleCount=15.0
 
 [EffectsQuality@3]
 r.TranslucencyLightingVolumeDim=64
@@ -450,8 +474,18 @@ r.SSS.HalfRes=1
 r.EmitterSpawnRateScale=1.0
 r.ParticleLightQuality=2
 r.DetailLevel=3
-r.SkyAtmosphere.FastSkyLUT=1
 r.SkyAtmosphere.AerialPerspectiveLUT.FastApplyOnOpaque=1 ; Always have FastSkyLUT 1 in this case to avoid wrong sky
+r.SkyAtmosphere.AerialPerspectiveLUT.SampleCountPerSlice=1
+r.SkyAtmosphere.AerialPerspectiveLUT.DepthResolution=16.0
+r.SkyAtmosphere.FastSkyLUT=1
+r.SkyAtmosphere.FastSkyLUT.SampleCountMin=4.0
+r.SkyAtmosphere.FastSkyLUT.SampleCountMax=32.0
+r.SkyAtmosphere.SampleCountMin=4.0
+r.SkyAtmosphere.SampleCountMax=32.0
+r.SkyAtmosphere.TransmittanceLUT.UseSmallFormat=0
+r.SkyAtmosphere.TransmittanceLUT.SampleCount=10.0
+r.SkyAtmosphere.MultiScatteringLUT.UseSmallFormat=0
+r.SkyAtmosphere.MultiScatteringLUT.SampleCount=15.0
 
 [EffectsQuality@Cine]
 r.TranslucencyLightingVolumeDim=64
@@ -468,13 +502,17 @@ r.SSS.HalfRes=0
 r.EmitterSpawnRateScale=1.0
 r.ParticleLightQuality=2
 r.DetailLevel=4
-r.SkyAtmosphere.FastSkyLUT=0
 r.SkyAtmosphere.AerialPerspectiveLUT.FastApplyOnOpaque=0
+r.SkyAtmosphere.AerialPerspectiveLUT.SampleCountPerSlice=2
+r.SkyAtmosphere.AerialPerspectiveLUT.DepthResolution=32.0
+r.SkyAtmosphere.FastSkyLUT=0
+r.SkyAtmosphere.FastSkyLUT.SampleCountMin=4.0
+r.SkyAtmosphere.FastSkyLUT.SampleCountMax=32.0
 r.SkyAtmosphere.SampleCountMin=8.0
 r.SkyAtmosphere.SampleCountMax=32.0
-r.SkyAtmosphere.AerialPerspectiveLUT.DepthResolution=32.0
-r.SkyAtmosphere.AerialPerspectiveLUT.SampleCountPerSlice=2.0
+r.SkyAtmosphere.TransmittanceLUT.UseSmallFormat=0
 r.SkyAtmosphere.TransmittanceLUT.SampleCount=30.0
+r.SkyAtmosphere.MultiScatteringLUT.UseSmallFormat=0
 r.SkyAtmosphere.MultiScatteringLUT.SampleCount=20.0
 
 ;-----------------------------------------------------------------------------------------------------------------
