Fix Android version checks in BaseDeviceProfiles.ini for Android 10

#ue4
#android

diff --git a/Engine/Config/BaseDeviceProfiles.ini b/Engine/Config/BaseDeviceProfiles.ini
--- a/Engine/Config/BaseDeviceProfiles.ini
+++ b/Engine/Config/BaseDeviceProfiles.ini
@@ -400,7 +400,7 @@ MatchProfile=(Profile="Android_PowerVR54x",Match=((SourceType=SRC_GpuFamily,Comp
 +MatchProfile=(Profile="Android_Adreno4xx_Nexus6",Match=((SourceType=SRC_DeviceModel,CompareType=CMP_Equal,MatchString="Nexus 6")))
 +MatchProfile=(Profile="Android_Adreno4xx",Match=((SourceType=SRC_GpuFamily,CompareType=CMP_Regex,MatchString="Adreno \\(TM\\) 4[0-9][0-9]")))
 +MatchProfile=(Profile="Android_Adreno51x",Match=((SourceType=SRC_GpuFamily,CompareType=CMP_Regex,MatchString="Adreno \\(TM\\) 51[0-9]")))
-+MatchProfile=(Profile="Android_Adreno5xx_No_Vulkan",Match=((SourceType=SRC_GpuFamily,CompareType=CMP_Regex,MatchString="Adreno \\(TM\\) 5[0-9][0-9]"),(SourceType=SRC_AndroidVersion, CompareType=CMP_Regex,MatchString="([0-9]).*"),(SourceType=SRC_PreviousRegexMatch,CompareType=CMP_Less,MatchString="8")))
++MatchProfile=(Profile="Android_Adreno5xx_No_Vulkan",Match=((SourceType=SRC_GpuFamily,CompareType=CMP_Regex,MatchString="Adreno \\(TM\\) 5[0-9][0-9]"),(SourceType=SRC_AndroidVersion, CompareType=CMP_Regex,MatchString="([0-9]+).*"),(SourceType=SRC_PreviousRegexMatch,CompareType=CMP_Less,MatchString="8")))
 +MatchProfile=(Profile="Android_Adreno5xx",Match=((SourceType=SRC_GpuFamily,CompareType=CMP_Regex,MatchString="Adreno \\(TM\\) 5[0-9][0-9]")))
 +MatchProfile=(Profile="Android_Adreno6xx",Match=((SourceType=SRC_GpuFamily,CompareType=CMP_Regex,MatchString="Adreno \\(TM\\) 6[0-9][0-9]")))
 +MatchProfile=(Profile="Android_TegraK1",Match=((SourceType=SRC_GpuFamily,CompareType=CMP_Equal,MatchString="NVIDIA Tegra"),(SourceType=SRC_GlVersion,CompareType=CMP_Regex,MatchString="^OpenGL ES 3\\.")))
@@ -416,7 +416,7 @@ MatchProfile=(Profile="Android_PowerVR54x",Match=((SourceType=SRC_GpuFamily,Comp
 +MatchProfile=(Profile="Android_Mali_T6xx",Match=((SourceType=SRC_GpuFamily,CompareType=CMP_Regex,MatchString="^Mali\\-T6")))
 +MatchProfile=(Profile="Android_Mali_T7xx",Match=((SourceType=SRC_GpuFamily,CompareType=CMP_Regex,MatchString="^Mali\\-T7")))
 +MatchProfile=(Profile="Android_Mali_T8xx_No_Vulkan",Match=((SourceType=SRC_GpuFamily,CompareType=CMP_Regex,MatchString="^Mali\\-T8"),(SourceType=SRC_VulkanVersion,CompareType=CMP_Regex,MatchString="^1\\.0\\.([0-9]+)"),(SourceType=SRC_PreviousRegexMatch,CompareType=CMP_Less,MatchString="11")))
-+MatchProfile=(Profile="Android_Mali_T8xx_No_Vulkan",Match=((SourceType=SRC_GpuFamily,CompareType=CMP_Regex,MatchString="^Mali\\-T8"),(SourceType=SRC_AndroidVersion, CompareType=CMP_Regex,MatchString="([0-9]).*"),(SourceType=SRC_PreviousRegexMatch,CompareType=CMP_Less,MatchString="8")))
++MatchProfile=(Profile="Android_Mali_T8xx_No_Vulkan",Match=((SourceType=SRC_GpuFamily,CompareType=CMP_Regex,MatchString="^Mali\\-T8"),(SourceType=SRC_AndroidVersion, CompareType=CMP_Regex,MatchString="([0-9]+).*"),(SourceType=SRC_PreviousRegexMatch,CompareType=CMP_Less,MatchString="8")))
 +MatchProfile=(Profile="Android_Mali_T8xx",Match=((SourceType=SRC_GpuFamily,CompareType=CMP_Regex,MatchString="^Mali\\-T8")))
 +MatchProfile=(Profile="Android_Mali_G71",Match=((SourceType=SRC_GpuFamily,CompareType=CMP_Regex,MatchString="^Mali\\-G71")))
 +MatchProfile=(Profile="Android_Mali_G72",Match=((SourceType=SRC_GpuFamily,CompareType=CMP_Regex,MatchString="^Mali\\-G72")))
