Disable small format for MultiScattering LUT on all platforms, as it causes different sky look with no measurable perf win

diff --git a/Engine/Config/Android/AndroidScalability.ini b/Engine/Config/Android/AndroidScalability.ini
--- a/Engine/Config/Android/AndroidScalability.ini
+++ b/Engine/Config/Android/AndroidScalability.ini
@@ -12,7 +12,6 @@ r.SkyAtmosphere.FastSkyLUT.SampleCountMax=8.0
 r.SkyAtmosphere.FastSkyLUT.Width=96
 r.SkyAtmosphere.FastSkyLUT.Height=50
 r.SkyAtmosphere.TransmittanceLUT.UseSmallFormat=1
-r.SkyAtmosphere.MultiScatteringLUT.UseSmallFormat=1
 
 [EffectsQuality@1]
 r.SkyAtmosphere.AerialPerspectiveLUT.FastApplyOnOpaque=1 ; Always have FastSkyLUT 1 in this case to avoid wrong sky
@@ -24,7 +23,6 @@ r.SkyAtmosphere.FastSkyLUT.SampleCountMax=8.0
 r.SkyAtmosphere.FastSkyLUT.Width=96
 r.SkyAtmosphere.FastSkyLUT.Height=50
 r.SkyAtmosphere.TransmittanceLUT.UseSmallFormat=1
-r.SkyAtmosphere.MultiScatteringLUT.UseSmallFormat=1
 
 [EffectsQuality@2]
 r.SkyAtmosphere.AerialPerspectiveLUT.FastApplyOnOpaque=1 ; Always have FastSkyLUT 1 in this case to avoid wrong sky
@@ -36,7 +34,6 @@ r.SkyAtmosphere.FastSkyLUT.SampleCountMax=8.0
 r.SkyAtmosphere.FastSkyLUT.Width=96
 r.SkyAtmosphere.FastSkyLUT.Height=50
 r.SkyAtmosphere.TransmittanceLUT.UseSmallFormat=1
-r.SkyAtmosphere.MultiScatteringLUT.UseSmallFormat=1
 
 [EffectsQuality@3]
 r.SkyAtmosphere.AerialPerspectiveLUT.FastApplyOnOpaque=1 ; Always have FastSkyLUT 1 in this case to avoid wrong sky
@@ -48,4 +45,3 @@ r.SkyAtmosphere.FastSkyLUT.SampleCountMax=8.0
 r.SkyAtmosphere.FastSkyLUT.Width=96
 r.SkyAtmosphere.FastSkyLUT.Height=50
 r.SkyAtmosphere.TransmittanceLUT.UseSmallFormat=1
-r.SkyAtmosphere.MultiScatteringLUT.UseSmallFormat=1
\ No newline at end of file
diff --git a/Engine/Config/BaseScalability.ini b/Engine/Config/BaseScalability.ini
--- a/Engine/Config/BaseScalability.ini
+++ b/Engine/Config/BaseScalability.ini
@@ -400,7 +400,6 @@ r.SkyAtmosphere.SampleCountMin=2.0
 r.SkyAtmosphere.SampleCountMax=16.0
 r.SkyAtmosphere.TransmittanceLUT.UseSmallFormat=1
 r.SkyAtmosphere.TransmittanceLUT.SampleCount=10.0
-r.SkyAtmosphere.MultiScatteringLUT.UseSmallFormat=1
 r.SkyAtmosphere.MultiScatteringLUT.SampleCount=15.0
 
 [EffectsQuality@1]
@@ -428,7 +427,6 @@ r.SkyAtmosphere.SampleCountMin=4.0
 r.SkyAtmosphere.SampleCountMax=32.0
 r.SkyAtmosphere.TransmittanceLUT.UseSmallFormat=0
 r.SkyAtmosphere.TransmittanceLUT.SampleCount=10.0
-r.SkyAtmosphere.MultiScatteringLUT.UseSmallFormat=0
 r.SkyAtmosphere.MultiScatteringLUT.SampleCount=15.0
 
 [EffectsQuality@2]
@@ -456,7 +454,6 @@ r.SkyAtmosphere.SampleCountMin=4.0
 r.SkyAtmosphere.SampleCountMax=32.0
 r.SkyAtmosphere.TransmittanceLUT.UseSmallFormat=0
 r.SkyAtmosphere.TransmittanceLUT.SampleCount=10.0
-r.SkyAtmosphere.MultiScatteringLUT.UseSmallFormat=0
 r.SkyAtmosphere.MultiScatteringLUT.SampleCount=15.0
 
 [EffectsQuality@3]
@@ -484,7 +481,6 @@ r.SkyAtmosphere.SampleCountMin=4.0
 r.SkyAtmosphere.SampleCountMax=32.0
 r.SkyAtmosphere.TransmittanceLUT.UseSmallFormat=0
 r.SkyAtmosphere.TransmittanceLUT.SampleCount=10.0
-r.SkyAtmosphere.MultiScatteringLUT.UseSmallFormat=0
 r.SkyAtmosphere.MultiScatteringLUT.SampleCount=15.0
 
 [EffectsQuality@Cine]
@@ -512,7 +508,6 @@ r.SkyAtmosphere.SampleCountMin=8.0
 r.SkyAtmosphere.SampleCountMax=32.0
 r.SkyAtmosphere.TransmittanceLUT.UseSmallFormat=0
 r.SkyAtmosphere.TransmittanceLUT.SampleCount=30.0
-r.SkyAtmosphere.MultiScatteringLUT.UseSmallFormat=0
 r.SkyAtmosphere.MultiScatteringLUT.SampleCount=20.0
 
 ;-----------------------------------------------------------------------------------------------------------------
diff --git a/Engine/Config/IOS/IOSScalability.ini b/Engine/Config/IOS/IOSScalability.ini
--- a/Engine/Config/IOS/IOSScalability.ini
+++ b/Engine/Config/IOS/IOSScalability.ini
@@ -12,7 +12,6 @@ r.SkyAtmosphere.FastSkyLUT.SampleCountMax=8.0
 r.SkyAtmosphere.FastSkyLUT.Width=96
 r.SkyAtmosphere.FastSkyLUT.Height=50
 r.SkyAtmosphere.TransmittanceLUT.UseSmallFormat=1
-r.SkyAtmosphere.MultiScatteringLUT.UseSmallFormat=1
 
 [EffectsQuality@1]
 r.SkyAtmosphere.AerialPerspectiveLUT.FastApplyOnOpaque=1 ; Always have FastSkyLUT 1 in this case to avoid wrong sky
@@ -24,7 +23,6 @@ r.SkyAtmosphere.FastSkyLUT.SampleCountMax=8.0
 r.SkyAtmosphere.FastSkyLUT.Width=96
 r.SkyAtmosphere.FastSkyLUT.Height=50
 r.SkyAtmosphere.TransmittanceLUT.UseSmallFormat=1
-r.SkyAtmosphere.MultiScatteringLUT.UseSmallFormat=1
 
 [EffectsQuality@2]
 r.SkyAtmosphere.AerialPerspectiveLUT.FastApplyOnOpaque=1 ; Always have FastSkyLUT 1 in this case to avoid wrong sky
@@ -36,7 +34,6 @@ r.SkyAtmosphere.FastSkyLUT.SampleCountMax=8.0
 r.SkyAtmosphere.FastSkyLUT.Width=96
 r.SkyAtmosphere.FastSkyLUT.Height=50
 r.SkyAtmosphere.TransmittanceLUT.UseSmallFormat=1
-r.SkyAtmosphere.MultiScatteringLUT.UseSmallFormat=1
 
 [EffectsQuality@3]
 r.SkyAtmosphere.AerialPerspectiveLUT.FastApplyOnOpaque=1 ; Always have FastSkyLUT 1 in this case to avoid wrong sky
@@ -48,4 +45,3 @@ r.SkyAtmosphere.FastSkyLUT.SampleCountMax=8.0
 r.SkyAtmosphere.FastSkyLUT.Width=96
 r.SkyAtmosphere.FastSkyLUT.Height=50
 r.SkyAtmosphere.TransmittanceLUT.UseSmallFormat=1
-r.SkyAtmosphere.MultiScatteringLUT.UseSmallFormat=1
\ No newline at end of file
