Add fx.NiagaraGlobalSpawnCountScale to Engine/Config/BaseScalability.ini and EngineTest/Config/DefaultScalability.ini

diff --git a/Engine/Config/BaseScalability.ini b/Engine/Config/BaseScalability.ini
--- a/Engine/Config/BaseScalability.ini
+++ b/Engine/Config/BaseScalability.ini
@@ -401,6 +401,7 @@ r.SkyAtmosphere.SampleCountMax=16.0
 r.SkyAtmosphere.TransmittanceLUT.UseSmallFormat=1
 r.SkyAtmosphere.TransmittanceLUT.SampleCount=10.0
 r.SkyAtmosphere.MultiScatteringLUT.SampleCount=15.0
+fx.NiagaraGlobalSpawnCountScale=0.125
 
 [EffectsQuality@1]
 r.TranslucencyLightingVolumeDim=32
@@ -428,6 +429,7 @@ r.SkyAtmosphere.SampleCountMax=32.0
 r.SkyAtmosphere.TransmittanceLUT.UseSmallFormat=0
 r.SkyAtmosphere.TransmittanceLUT.SampleCount=10.0
 r.SkyAtmosphere.MultiScatteringLUT.SampleCount=15.0
+fx.NiagaraGlobalSpawnCountScale=0.25
 
 [EffectsQuality@2]
 r.TranslucencyLightingVolumeDim=48
@@ -455,6 +457,7 @@ r.SkyAtmosphere.SampleCountMax=32.0
 r.SkyAtmosphere.TransmittanceLUT.UseSmallFormat=0
 r.SkyAtmosphere.TransmittanceLUT.SampleCount=10.0
 r.SkyAtmosphere.MultiScatteringLUT.SampleCount=15.0
+fx.NiagaraGlobalSpawnCountScale=0.5
 
 [EffectsQuality@3]
 r.TranslucencyLightingVolumeDim=64
@@ -482,6 +485,7 @@ r.SkyAtmosphere.SampleCountMax=32.0
 r.SkyAtmosphere.TransmittanceLUT.UseSmallFormat=0
 r.SkyAtmosphere.TransmittanceLUT.SampleCount=10.0
 r.SkyAtmosphere.MultiScatteringLUT.SampleCount=15.0
+fx.NiagaraGlobalSpawnCountScale=1.0
 
 [EffectsQuality@Cine]
 r.TranslucencyLightingVolumeDim=64
@@ -509,6 +513,7 @@ r.SkyAtmosphere.SampleCountMax=32.0
 r.SkyAtmosphere.TransmittanceLUT.UseSmallFormat=0
 r.SkyAtmosphere.TransmittanceLUT.SampleCount=30.0
 r.SkyAtmosphere.MultiScatteringLUT.SampleCount=20.0
+fx.NiagaraGlobalSpawnCountScale=1.0
 
 ;-----------------------------------------------------------------------------------------------------------------
 
