HTML5 platform extension - fix runtime error on non-HTML5 platforms

#rn

diff --git a/Engine/Config/BaseEngine.ini b/Engine/Config/BaseEngine.ini
--- a/Engine/Config/BaseEngine.ini
+++ b/Engine/Config/BaseEngine.ini
@@ -2254,25 +2254,6 @@ PktIncomingLagMax=200
 +NetworkEmulationProfiles=(ProfileName="Average",ToolTip="Simulates average internet conditions")
 +NetworkEmulationProfiles=(ProfileName="Bad",ToolTip="Simulates laggy internet conditions")
 
-[/Script/HTML5Networking.WebSocketNetDriver]
-AllowPeerConnections=False
-AllowPeerVoice=False
-ConnectionTimeout=60.0
-InitialConnectTimeout=120.0
-RecentlyDisconnectedTrackingTime=180
-AckTimeout=10.0
-KeepAliveTime=20.2
-MaxClientRate=15000
-MaxInternetClientRate=10000
-RelevantTimeout=5.0
-SpawnPrioritySeconds=1.0
-ServerTravelPause=4.0
-NetServerMaxTickRate=30
-LanServerMaxTickRate=35
-WebSocketPort=8889
-NetConnectionClassName="/Script/HTML5Networking.WebSocketConnection"
-MaxPortCountToTry=512
-
 [/Script/GameplayDebugger.GameplayDebuggingControllerComponent]
 CategoryZeroBind=(Key=NumPadZero,bShift=False,bCtrl=False,bAlt=False,bCmd=False)
 CategoryOneBind=(Key=NumPadOne,bShift=False,bCtrl=False,bAlt=False,bCmd=False)
