Add deprecated Remote Channel to base engine.

diff --git a/Engine/Config/BaseEngine.ini b/Engine/Config/BaseEngine.ini
--- a/Engine/Config/BaseEngine.ini
+++ b/Engine/Config/BaseEngine.ini
@@ -2377,6 +2377,9 @@ r.GPUCrashDebugging=false
 [Messaging]
 bAllowDelayedMessaging=false
 
+[RemoteSession]
++ChannelRedirects=(OldName=FRemoteSessionFrameBufferChannel,NewName=FRemoteSessionFrameBufferChannel_DEPRECATED)
+
 [/Script/ChaosSolverEngine.ChaosSolverSettings]
 DefaultChaosSolverActorClass=/Script/ChaosSolverEngine.ChaosSolverActor
 
