Merging CL 4573359 from //UE4/Release-4.21
PR #5265: Removed duplicated WindowsMixedRealityHMD entry (Contributed by [REDACTED_PERSON])

diff --git a/Engine/Config/BaseEngine.ini b/Engine/Config/BaseEngine.ini
--- a/Engine/Config/BaseEngine.ini
+++ b/Engine/Config/BaseEngine.ini
@@ -1959,7 +1959,6 @@ WindowsMixedRealityHMD=40
 OculusHMD=20
 SteamVR=10
 OSVR=5
-WindowsMixedRealityHMD=3
 
 [/Script/Engine.AISystemBase]
 AISystemModuleName=AIModule
