Rollback change to BaseDeviceProfiles.ini from CL-4636504 'Updated ML integration to MLSDK 0.19"

Something was not set up correctly in one of the device profiles for Lumin, and it broke half of our automated rendering tests for windows.

diff --git a/Engine/Config/BaseDeviceProfiles.ini b/Engine/Config/BaseDeviceProfiles.ini
--- a/Engine/Config/BaseDeviceProfiles.ini
+++ b/Engine/Config/BaseDeviceProfiles.ini
@@ -603,42 +603,7 @@ DeviceType=Android
 +CVars=r.Lumin.OverrideExternalTextureSupport=4
 +CVars=r.Vulkan.DelayAcquireBackBuffer=1
 +CVars=vr.DebugCanvasInLayer=1
-+CVars=r.Vulkan.RHIThread=1
-+CVars=Slate.EnableFastWidgetPath=0
-+CVars=r.DisjointTimerQueries=0
-+CVars=r.MobileContentScaleFactor=1
-+CVars=r.BloomQuality=0
-+CVars=r.DepthOfFieldQuality=0
-+CVars=r.LightShaftQuality=0
-+CVars=r.RefractionQuality=0
-+CVars=r.ShadowQuality=2
-; image tracker crashes if this is set to 1
-+CVars=core.Lumin.UseThreadingForPerformance=0
-
-[Lumin_Desktop DeviceProfile]
-DeviceType=Lumin
-BaseProfileName=Lumin
-+CVars=r.BloomQuality=1
-+CVars=r.DepthOfFieldQuality=1
-+CVars=r.LightShaftQuality=1
-+CVars=r.RefractionQuality=1
-+CVars=r.MobileContentScaleFactor=1.0
-; Disable Temporal AA - would need to be adjusted for an AR system to know where the background is
-+CVars=r.PostProcessAAQuality=0
-; Disable Screen Space Reflection
-+CVars=r.SSR.Quality=0
-; Disable Screen Space Ambient Occlusion
-+CVars=r.AmbientOcclusionLevels=0
-; Disable Reflection Environment
-+CVars=r.ReflectionEnvironment=0
-; Disable MotionBlur
-+CVars=r.MotionBlurQuality=0
-; Disable Light Propagation Volume
-+CVars=r.LPV.Intensity=0
-; Disable HZB
-+CVars=r.HZBOcclusion=0
-+CVars=r.TranslucentLightingVolume=0
-+CVars=r.AllowOcclusionQueries=0
++CVars=r.Vulkan.RHIThread=0
 
 [HTML5 DeviceProfile]
 DeviceType=HTML5
