Lumin device profile update to fix leaking profiles when not running on device.

diff --git a/Engine/Config/BaseDeviceProfiles.ini b/Engine/Config/BaseDeviceProfiles.ini
--- a/Engine/Config/BaseDeviceProfiles.ini
+++ b/Engine/Config/BaseDeviceProfiles.ini
@@ -82,8 +82,8 @@
 +DeviceProfileNameAndTypes=Android_Mali_G71,Android
 +DeviceProfileNameAndTypes=Android_Mali_G72,Android
 +DeviceProfileNameAndTypes=Android_Mali_G76,Android
-; @todo Lumin: Once we get the IniPlatformParent stuff, this should move into a LuminBaseDeviceProfiles.ini, and make type be Lumin
-+DeviceProfileNameAndTypes=Lumin,Android
++DeviceProfileNameAndTypes=Lumin,Lumin
++DeviceProfileNameAndTypes=Lumin_Desktop,Lumin
 
 [/Script/Engine.TextureLODSettings]
 @TextureLODGroups=Group
@@ -597,13 +597,51 @@ BaseProfileName=Android_Mali_T8xx
 +CVars=r.Android.DisableVulkanSupport=1
 
 [Lumin DeviceProfile]
-; @todo Lumin: Android - no!
-DeviceType=Android
+DeviceType=Lumin
+BaseProfileName=
 +CVars=r.defaultfeature.postprocessing=0
 +CVars=r.Lumin.OverrideExternalTextureSupport=4
 +CVars=r.Vulkan.DelayAcquireBackBuffer=1
 +CVars=vr.DebugCanvasInLayer=1
-+CVars=r.Vulkan.RHIThread=0
++CVars=r.Vulkan.RHIThread=1
++CVars=Slate.EnableFastWidgetPath=0
++CVars=r.DisjointTimerQueries=0
++CVars=r.MobileContentScaleFactor=1
++CVars=r.BloomQuality=0
++CVars=r.DepthOfFieldQuality=0
++CVars=r.LightShaftQuality=0
++CVars=r.RefractionQuality=0
++CVars=r.ShadowQuality=2
+
+[LuminClient DeviceProfile]
+DeviceType=Lumin
+BaseProfileName=Lumin
+
+[Lumin_Desktop DeviceProfile]
+DeviceType=Lumin
+BaseProfileName=Lumin
++CVars=r.BloomQuality=1
++CVars=r.DepthOfFieldQuality=1
++CVars=r.LightShaftQuality=1
++CVars=r.RefractionQuality=1
++CVars=r.MobileContentScaleFactor=1.0
+; Disable Temporal AA - would need to be adjusted for an AR system to know where the background is
++CVars=r.PostProcessAAQuality=0
+; Disable Screen Space Reflection
++CVars=r.SSR.Quality=0
+; Disable Screen Space Ambient Occlusion
++CVars=r.AmbientOcclusionLevels=0
+; Disable Reflection Environment
++CVars=r.ReflectionEnvironment=0
+; Disable MotionBlur
++CVars=r.MotionBlurQuality=0
+; Disable Light Propagation Volume
++CVars=r.LPV.Intensity=0
+; Disable HZB
++CVars=r.HZBOcclusion=0
++CVars=r.TranslucentLightingVolume=0
++CVars=r.AllowOcclusionQueries=0
++CVars=r.DepthOfFieldQuality=0
 
 [HTML5 DeviceProfile]
 DeviceType=HTML5
diff --git a/Engine/Plugins/Lumin/MagicLeap/Source/MagicLeap/Private/MagicLeapHMD.cpp b/Engine/Plugins/Lumin/MagicLeap/Source/MagicLeap/Private/MagicLeapHMD.cpp
--- a/Engine/Plugins/Lumin/MagicLeap/Source/MagicLeap/Private/MagicLeapHMD.cpp
+++ b/Engine/Plugins/Lumin/MagicLeap/Source/MagicLeap/Private/MagicLeapHMD.cpp
@@ -1243,8 +1243,11 @@ void FMagicLeapHMD::EnableDeviceFeatures()
 		EnablePerception();
 		EnableHeadTracking();
 		EnableInputDevices();
+
+		// We also avoid enabling the custom profile when there's no HMD, as otherwise
+		// we get the profile effects on non-vr-preview rendering.
+		EnableLuminProfile();
 	}
-	EnableLuminProfile();
 }
 
 void FMagicLeapHMD::DisableDeviceFeatures()
diff --git a/Engine/Source/Programs/AutomationTool/Lumin/LuminPlatform.Automation.cs b/Engine/Source/Programs/AutomationTool/Lumin/LuminPlatform.Automation.cs
--- a/Engine/Source/Programs/AutomationTool/Lumin/LuminPlatform.Automation.cs
+++ b/Engine/Source/Programs/AutomationTool/Lumin/LuminPlatform.Automation.cs
@@ -26,8 +26,7 @@ public class LuminPlatform : Platform
 	public LuminPlatform()
 		: base(UnrealTargetPlatform.Lumin)
 	{
-		// @todo Lumin: once we get ini subplatforms, fix this and also TVOS stuff!
-		TargetIniPlatformType = UnrealTargetPlatform.Android;
+		TargetIniPlatformType = UnrealTargetPlatform.Lumin;
 		RuntimeDependenciesForMabu = new List<FileReference>();
 	}
 
@@ -65,8 +64,7 @@ public class LuminPlatform : Platform
 	private string StageIconFileToMabu(string ConfigPropertyName, string IconStagePath, DeploymentContext SC)
 	{
 		// Read in any extra assets required to correctly install the application.
-		// @todo Lumin: subclass ini platforms needs to use Android for now
-		ConfigHierarchy Ini = ConfigCache.ReadHierarchy(ConfigHierarchyType.Engine, SC.RawProjectPath.Directory, UnrealTargetPlatform.Android);
+		ConfigHierarchy Ini = ConfigCache.ReadHierarchy(ConfigHierarchyType.Engine, SC.RawProjectPath.Directory, UnrealTargetPlatform.Lumin);
 		// ConfigHierarchy Ini = ConfigCache.ReadHierarchy(ConfigHierarchyType.Engine, SC.RawProjectPath.Directory, SC.StageTargetPlatform.PlatformType);
 		string Value;
 		Ini.GetString("/Script/LuminRuntimeSettings.LuminRuntimeSettings", ConfigPropertyName, out Value);
diff --git a/Engine/Source/Programs/UnrealBuildTool/Platform/Lumin/UEBuildLumin.cs b/Engine/Source/Programs/UnrealBuildTool/Platform/Lumin/UEBuildLumin.cs
--- a/Engine/Source/Programs/UnrealBuildTool/Platform/Lumin/UEBuildLumin.cs
+++ b/Engine/Source/Programs/UnrealBuildTool/Platform/Lumin/UEBuildLumin.cs
@@ -86,7 +86,6 @@ namespace UnrealBuildTool
 			}
 
 			// look up Android specific settings
-			// @todo Lumin: When we subclass platform ini's, this would be Platform!!
 			if (!DoProjectSettingsMatchDefault(UnrealTargetPlatform.Lumin, ProjectPath, "/Script/LuminRuntimeSettings.LuminRuntimeSettings", ConfigBoolKeys.ToArray(), null, null))
 			{
 				return false;
diff --git a/Engine/Source/Programs/UnrealBuildTool/Platform/Lumin/UEDeployLumin.cs b/Engine/Source/Programs/UnrealBuildTool/Platform/Lumin/UEDeployLumin.cs
--- a/Engine/Source/Programs/UnrealBuildTool/Platform/Lumin/UEDeployLumin.cs
+++ b/Engine/Source/Programs/UnrealBuildTool/Platform/Lumin/UEDeployLumin.cs
@@ -32,9 +32,7 @@ namespace UnrealBuildTool
 
 		private ConfigHierarchy GetConfigCacheIni(ConfigHierarchyType Type)
 		{
-			// @todo Lumin: So - this is the problem with subclassing a platform currently - ini files. Lumin will use Android ini files
-			// until I finish and get code over from another branch (key letter Q) that allows for insertion of a subclassed ini platform thing
-			return ConfigCache.ReadHierarchy(Type, DirectoryReference.FromFile(ProjectFile), UnrealTargetPlatform.Android);
+			return ConfigCache.ReadHierarchy(Type, DirectoryReference.FromFile(ProjectFile), UnrealTargetPlatform.Lumin);
 		}
 
 		private string GetRuntimeSetting(string Key)
