DR - Add cvar to help debug (disabled)
#rnx

diff --git a/Engine/Config/ConsoleVariables.ini b/Engine/Config/ConsoleVariables.ini
--- a/Engine/Config/ConsoleVariables.ini
+++ b/Engine/Config/ConsoleVariables.ini
@@ -39,6 +39,8 @@
 ;r.DumpShaderDebugInfo=1
 ; When this is enabled, dumped shader paths will get collapsed (in the cases where paths are longer than the OS's max)
 ;r.DumpShaderDebugShortNames=1
+; When this is enabled, SCW crashes will print out the list of jobs in the current worker
+;r.DumpSCWQueuedJobs=1
 ; When this is enabled, when dumping shaders an additional file to use with ShaderCompilerWorker -direct mode will be generated
 ;r.DumpShaderDebugWorkerCommandLine=1
 ; Uncomment to enable parallel rendering at startup
