Cranks up the number of pixel that can be scattered for DOF to 25% on cinematic settings.

diff --git a/Engine/Config/BaseScalability.ini b/Engine/Config/BaseScalability.ini
--- a/Engine/Config/BaseScalability.ini
+++ b/Engine/Config/BaseScalability.ini
@@ -311,7 +311,7 @@ r.DOF.Gather.RingCount=5                 ; high number of samples when gathering
 r.DOF.Scatter.ForegroundCompositing=1    ; additive foreground scattering
 r.DOF.Scatter.BackgroundCompositing=2    ; background scattering occlusion
 r.DOF.Scatter.EnableBokehSettings=1      ; no bokeh simulation when scattering
-r.DOF.Scatter.MaxSpriteRatio=0.1         ; only a maximum of 10% of scattered bokeh
+r.DOF.Scatter.MaxSpriteRatio=0.25        ; only a maximum of 10% of scattered bokeh
 r.DOF.Recombine.Quality=2                ; highest slight out of focus
 r.DOF.Recombine.EnableBokehSettings=1    ; bokeh simulation on slight out of focus
 r.DOF.TemporalAAQuality=1                ; more stable temporal accumulation
