Merge fix, add stuff back to BaseEngine.ini

diff --git a/Engine/Config/BaseEngine.ini b/Engine/Config/BaseEngine.ini
--- a/Engine/Config/BaseEngine.ini
+++ b/Engine/Config/BaseEngine.ini
@@ -88,7 +88,9 @@ DefaultBokehTextureName=/Engine/EngineMaterials/DefaultBokeh.DefaultBokeh
 DefaultBloomKernelTextureName=/Engine/EngineMaterials/DefaultBloomKernel.DefaultBloomKernel
 GeomMaterialName=/Engine/EngineDebugMaterials/GeomMaterial.GeomMaterial
 DebugMeshMaterialName=/Engine/EngineDebugMaterials/DebugMeshMaterial.DebugMeshMaterial
+EmissiveMeshMaterialName=/Engine/EngineMaterials/EmissiveMeshMaterial.EmissiveMeshMaterial
 PreIntegratedSkinBRDFTextureName=/Engine/EngineMaterials/PreintegratedSkinBRDF.PreintegratedSkinBRDF
+BlueNoiseTextureName=/Engine/EngineMaterials/BlueNoise.BlueNoise
 MiniFontTextureName=/Engine/EngineMaterials/MiniFont.MiniFont
 WeightMapPlaceholderTextureName=/Engine/EngineMaterials/WeightMapPlaceholderTexture.WeightMapPlaceholderTexture
 LightMapDensityTextureName=/Engine/EngineMaterials/DefaultWhiteGrid.DefaultWhiteGrid
@@ -114,6 +116,7 @@ EditorBrushMaterialName=/Engine/EngineMaterials/EditorBrushMaterial.EditorBrushM
 DefaultPhysMaterialName=/Engine/EngineMaterials/DefaultPhysicalMaterial.DefaultPhysicalMaterial
 DefaultDeferredDecalMaterialName=/Engine/EngineMaterials/DefaultDeferredDecalMaterial.DefaultDeferredDecalMaterial
 DefaultPostProcessMaterialName=/Engine/EngineMaterials/DefaultPostProcessMaterial.DefaultPostProcessMaterial
+DefaultVirtualTextureMaterialName=/Engine/EngineMaterials/DefaultVirtualTextureMaterial.DefaultVirtualTextureMaterial
 TimecodeProviderClassName=None
 DefaultTimecodeProviderClassName=/Script/Engine.SystemTimeTimecodeProvider
 TextureStreamingBoundsMaterialName=/Engine/EditorMaterials/Utilities/TextureStreamingBounds_MATInst.TextureStreamingBounds_MATInst
@@ -730,6 +733,7 @@ bFullyCompressStartupPackages=false
 +Package=/Engine/EngineMaterials/DefaultPhysicalMaterial
 +Package=/Engine/EngineMaterials/DefaultWhiteGrid
 +Package=/Engine/EngineMaterials/EditorBrushMaterial
++Package=/Engine/EngineMaterials/EmissiveMeshMaterial
 +Package=/Engine/EngineMaterials/Good64x64TilingNoiseHighFreq
 +Package=/Engine/EngineMaterials/Grid
 +Package=/Engine/EngineMaterials/Grid_N
@@ -1608,6 +1612,11 @@ Project=(Type=FileSystem, ReadOnly=false, Clean=false, Flush=true, PurgeTransien
 Shared=(Type=FileSystem, ReadOnly=false, Clean=false, Flush=false, DeleteUnused=false, UnusedFileAge=23, FoldersToClean=-1, Path=?EpicDDC, EnvPathOverride=UE-SharedDataCachePath)
 AltShared=(Type=FileSystem, ReadOnly=true, Clean=false, Flush=false, DeleteUnused=false, UnusedFileAge=23, FoldersToClean=-1, Path=?EpicDDC2, EnvPathOverride=UE-SharedDataCachePath2)
 
+[VirtualTextureChunkDDCCache]
+UnusedFileAge=34
+MaxFileChecksPerSec=-1 ;no limit
+Path=%GAMEDIR%DerivedDataCache/VT
+
 [/Script/Engine.LocalPlayer]
 AspectRatioAxisConstraint=AspectRatio_MaintainXFOV
 
@@ -2275,3 +2284,19 @@ r.GPUCrashDebugging=false
 
 [Messaging]
 bAllowDelayedMessaging=false
+
+[/Script/Engine.VirtualTexturePoolConfig]
+; Configure VT physical memory pools. Usefull commands to set these
+; "stat virtualtexturing" show dynamic use of the VT system including the cache loads
+; "r.VT.ListPhysicalPool" list in-depth details on the allocated physical pools
+; Size in megabyte of any pools that are created and not explicitly defined below
+DefaultSizeInMegabyte=64
+; Explicit per pool configurations
++Pools=(SizeInMegabyte=128, TileSize=136, Format=PF_DXT1)
++Pools=(SizeInMegabyte=64, TileSize=136, Format=PF_DXT5)
++Pools=(SizeInMegabyte=128, TileSize=136, Format=PF_BC5)
++Pools=(SizeInMegabyte=32, TileSize=136, Format=PF_FloatRGBA)
++Pools=(SizeInMegabyte=32, TileSize=136, Format=PF_B8G8R8A8)
+
+
+
