Disable frame rate smoothing by default.

- Seems to be an artifact of 332560, no longer relevant.
- This setting affects a lot of Unreal games, see https://www.google.com/search?q=Unreal+fps+62  and https://www.google.com/search?q=bSmoothFrameRate
- Most other platforms already do this in their .inis

diff --git a/Engine/Config/BaseEngine.ini b/Engine/Config/BaseEngine.ini
--- a/Engine/Config/BaseEngine.ini
+++ b/Engine/Config/BaseEngine.ini
@@ -227,7 +227,7 @@ NetDriverDefinitions=(DefName="GameNetDriver",DriverClassName="/Script/OnlineSub
 NetErrorLogInterval=1.0
 SerializationOutOfBoundsErrorMessage=NSLOCTEXT("","SerializationOutOfBoundsErrorMessage","Corrupt data found, please verify your installation.")
 SerializationOutOfBoundsErrorMessageCaption=NSLOCTEXT("","SerializationOutOfBoundsErrorMessageCaption","Serialization Error : Action Needed")
-bSmoothFrameRate=true
+bSmoothFrameRate=false
 SmoothedFrameRateRange=(LowerBound=(Type="ERangeBoundTypes::Inclusive",Value=22),UpperBound=(Type="ERangeBoundTypes::Exclusive",Value=62))
 bCheckForMultiplePawnsSpawnedInAFrame=false
 NumPawnsAllowedToBeSpawnedInAFrame=2
