#rnx

diff --git a/Engine/Config/BasePakFileRules.ini b/Engine/Config/BasePakFileRules.ini
--- a/Engine/Config/BasePakFileRules.ini
+++ b/Engine/Config/BasePakFileRules.ini
@@ -14,7 +14,7 @@
 ; Exclude specific large textures on mobile platforms, this was moved from CopyBuildToStagingDirectory.cs
 ; This can be added to in a game's DefaultPakFileRules.ini by using the same sections, and new sections can be added
 ; To remove this rule, use !Files to clear the list of file paths
-Platforms="Android,iOS,tvOS,HTML5"
+Platforms="Android,iOS,tvOS"
 bExcludeFromPaks=true
 bOverrideChunkManifest=true
 +Files=".../Engine/Content/EngineMaterials/DefaultBloomKernel.*"
diff --git a/Engine/Platforms/HTML5/Config/BaseHTML5PakFileRules.ini b/Engine/Platforms/HTML5/Config/BaseHTML5PakFileRules.ini
new file mode 100644
--- /dev/null
+++ b/Engine/Platforms/HTML5/Config/BaseHTML5PakFileRules.ini
@@ -0,0 +1,5 @@
+[ExcludeContentForMobile]
+Platforms="HTML5"
+;bExcludeFromPaks=true
+;bOverrideChunkManifest=true
+;+Files=".../Engine/Content/EngineMaterials/DefaultBloomKernel.*"
