#rnx

diff --git a/Engine/Config/BaseDeviceProfiles.ini b/Engine/Config/BaseDeviceProfiles.ini
--- a/Engine/Config/BaseDeviceProfiles.ini
+++ b/Engine/Config/BaseDeviceProfiles.ini
@@ -53,7 +53,6 @@
 +DeviceProfileNameAndTypes=iPadMini5,IOS
 +DeviceProfileNameAndTypes=AppleTV,IOS
 +DeviceProfileNameAndTypes=AppleTV4K,IOS
-+DeviceProfileNameAndTypes=HTML5,HTML5
 +DeviceProfileNameAndTypes=Mac,Mac
 +DeviceProfileNameAndTypes=MacClient,Mac
 +DeviceProfileNameAndTypes=MacNoEditor,Mac
@@ -699,11 +698,6 @@ BaseProfileName=Lumin
 +CVars=r.AllowOcclusionQueries=0
 +CVars=r.DepthOfFieldQuality=0
 
-[HTML5 DeviceProfile]
-DeviceType=HTML5
-BaseProfileName=
-+CVars=r.RefractionQuality=1
-
 [Mac DeviceProfile]
 DeviceType=Mac
 BaseProfileName=
diff --git a/Engine/Platforms/HTML5/Config/BaseHTML5DeviceProfiles.ini b/Engine/Platforms/HTML5/Config/BaseHTML5DeviceProfiles.ini
new file mode 100644
--- /dev/null
+++ b/Engine/Platforms/HTML5/Config/BaseHTML5DeviceProfiles.ini
@@ -0,0 +1,8 @@
+[DeviceProfiles]
++DeviceProfileNameAndTypes=HTML5,HTML5
+
+[HTML5 DeviceProfile]
+DeviceType=HTML5
+BaseProfileName=
++CVars=r.RefractionQuality=1
+
