#rnx

diff --git a/Engine/Config/BaseEngine.ini b/Engine/Config/BaseEngine.ini
--- a/Engine/Config/BaseEngine.ini
+++ b/Engine/Config/BaseEngine.ini
@@ -2192,46 +2192,6 @@ NameForGameView3="GameView3"
 NameForGameView4="GameView4"
 NameForGameView5="GameView5"
 
-[/Script/HTML5PlatformEditor.HTML5TargetSettings]
-Compressed=False
-EnableIndexedDB=False
-EnableSIMD=False
-DeployServerPort=53501
-EnableTracing=False
-CanvasScalingMode=Aspect
-EnableMultithreading=False
-OffscreenCanvas=False
-LLVMWasmBackend=False
-
-; common browser locations - add your custom locations here
-[/Script/HTML5PlatformEditor.HTML5Browsers]
-; windows
-+BrowserLauncher=(BrowserName="FF Nightly(32bit)", BrowserPath=(FilePath="C:/Program Files (x86)/Firefox Nightly/firefox.exe"))
-+BrowserLauncher=(BrowserName="FF Dev(32bit)",     BrowserPath=(FilePath="C:/Program Files (x86)/Firefox Developer Edition/firefox.exe"))
-+BrowserLauncher=(BrowserName="Firefox(32bit)",    BrowserPath=(FilePath="C:/Program Files (x86)/Mozilla Firefox/firefox.exe"))
-+BrowserLauncher=(BrowserName="FF Nightly",        BrowserPath=(FilePath="C:/Program Files/Firefox Nightly/firefox.exe"))
-+BrowserLauncher=(BrowserName="FF Dev",            BrowserPath=(FilePath="C:/Program Files/Firefox Developer Edition/firefox.exe"))
-+BrowserLauncher=(BrowserName="Firefox",           BrowserPath=(FilePath="C:/Program Files/Mozilla Firefox/firefox.exe"))
-; old FF Nightly path...
-+BrowserLauncher=(BrowserName="FF Nightly(32bit)", BrowserPath=(FilePath="C:/Program Files (x86)/Nightly/firefox.exe"))
-+BrowserLauncher=(BrowserName="FF Nightly",        BrowserPath=(FilePath="C:/Program Files/Nightly/firefox.exe"))
-; chrome
-+BrowserLauncher=(BrowserName="Chrome",         BrowserPath=(FilePath="C:/Program Files (x86)/Google/Chrome/Application/chrome.exe"))
-; Chrome Canary -- edit path to canary installation location
-+BrowserLauncher=(BrowserName="Canary",            BrowserPath=(FilePath="[REDACTED_USER_HOME]/AppData/Local/Google/Chrome SxS/Application/chrome.exe"))
-; Edge -- will be autodetected...
-;+BrowserLauncher=(BrowserName="Edge",           BrowserPath=(FilePath="start microsoft-edge:"))
-
-; osx
-+BrowserLauncher=(BrowserName="Safari",  BrowserPath=(FilePath="/Applications/Safari.app"))
-+BrowserLauncher=(BrowserName="Firefox", BrowserPath=(FilePath="/Applications/Firefox.app"))
-+BrowserLauncher=(BrowserName="Chrome",  BrowserPath=(FilePath="/Applications/Google Chrome.app"))
-
-; linux
-+BrowserLauncher=(BrowserName="Firefox", BrowserPath=(FilePath="/usr/bin/firefox"))
-+BrowserLauncher=(BrowserName="Chrome", BrowserPath=(FilePath="/usr/bin/google-chrome"))
-+BrowserLauncher=(BrowserName="Chromium", BrowserPath=(FilePath="/usr/bin/chromium-browser"))
-
 [Browser]
 bForceMessageLoop=true
 
diff --git a/Engine/Platforms/HTML5/Config/BaseHTML5Engine.ini b/Engine/Platforms/HTML5/Config/BaseHTML5Engine.ini
new file mode 100644
--- /dev/null
+++ b/Engine/Platforms/HTML5/Config/BaseHTML5Engine.ini
@@ -0,0 +1,45 @@
+[/Script/HTML5PlatformEditor.HTML5TargetSettings]
+Compressed=False
+EnableIndexedDB=False
+EnableSIMD=False
+DeployServerPort=53501
+EnableTracing=False
+CanvasScalingMode=Aspect
+EnableMultithreading=False
+OffscreenCanvas=False
+LLVMWasmBackend=False
+
+
+; common browser locations - add your custom locations here
+[/Script/HTML5PlatformEditor.HTML5Browsers]
+; windows
++BrowserLauncher=(BrowserName="FF Nightly(32bit)", BrowserPath=(FilePath="C:/Program Files (x86)/Firefox Nightly/firefox.exe"))
++BrowserLauncher=(BrowserName="FF Dev(32bit)",     BrowserPath=(FilePath="C:/Program Files (x86)/Firefox Developer Edition/firefox.exe"))
++BrowserLauncher=(BrowserName="Firefox(32bit)",    BrowserPath=(FilePath="C:/Program Files (x86)/Mozilla Firefox/firefox.exe"))
++BrowserLauncher=(BrowserName="FF Nightly",        BrowserPath=(FilePath="C:/Program Files/Firefox Nightly/firefox.exe"))
++BrowserLauncher=(BrowserName="FF Dev",            BrowserPath=(FilePath="C:/Program Files/Firefox Developer Edition/firefox.exe"))
++BrowserLauncher=(BrowserName="Firefox",           BrowserPath=(FilePath="C:/Program Files/Mozilla Firefox/firefox.exe"))
+; old FF Nightly path...
++BrowserLauncher=(BrowserName="FF Nightly(32bit)", BrowserPath=(FilePath="C:/Program Files (x86)/Nightly/firefox.exe"))
++BrowserLauncher=(BrowserName="FF Nightly",        BrowserPath=(FilePath="C:/Program Files/Nightly/firefox.exe"))
+; chrome
++BrowserLauncher=(BrowserName="Chrome",         BrowserPath=(FilePath="C:/Program Files (x86)/Google/Chrome/Application/chrome.exe"))
+; Chrome Canary -- edit path to canary installation location
++BrowserLauncher=(BrowserName="Canary",            BrowserPath=(FilePath="[REDACTED_USER_HOME]/AppData/Local/Google/Chrome SxS/Application/chrome.exe"))
+; Edge -- will be autodetected...
+;+BrowserLauncher=(BrowserName="Edge",           BrowserPath=(FilePath="start microsoft-edge:"))
+
+; osx
++BrowserLauncher=(BrowserName="Safari",  BrowserPath=(FilePath="/Applications/Safari.app"))
++BrowserLauncher=(BrowserName="Firefox", BrowserPath=(FilePath="/Applications/Firefox.app"))
++BrowserLauncher=(BrowserName="Chrome",  BrowserPath=(FilePath="/Applications/Google Chrome.app"))
+
+; linux
++BrowserLauncher=(BrowserName="Firefox", BrowserPath=(FilePath="/usr/bin/firefox"))
++BrowserLauncher=(BrowserName="Chrome", BrowserPath=(FilePath="/usr/bin/google-chrome"))
++BrowserLauncher=(BrowserName="Chromium", BrowserPath=(FilePath="/usr/bin/chromium-browser"))
+
+
+;[OnlineSubsystem]
+;NativePlatformService=WebSockets
+
