Fixing the tick rate with SteamSockets on dedicated servers.

#rnx

diff --git a/Engine/Config/BaseEngine.ini b/Engine/Config/BaseEngine.ini
--- a/Engine/Config/BaseEngine.ini
+++ b/Engine/Config/BaseEngine.ini
@@ -1392,6 +1392,8 @@ NetConnectionClassName="/Script/OnlineSubsystemSteam.SteamNetConnection"
 NetConnectionClassName="/Script/SteamSockets.SteamSocketsNetConnection"
 ConnectionTimeout=80.0
 InitialConnectTimeout=120.0
+NetServerMaxTickRate=30
+MaxNetTickRate=120
 
 [OnlineSubsystemMcp]
 bEnabled=false
