Update BaseInput.ini to include new 'Class' enum value for dynamic (i.e. hotfix) volume

diff --git a/Engine/Config/BaseInput.ini b/Engine/Config/BaseInput.ini
--- a/Engine/Config/BaseInput.ini
+++ b/Engine/Config/BaseInput.ini
@@ -134,12 +134,12 @@ MaxScrollbackSize=1024
 
 ; From AudioDevice.cpp
 +ManualAutoCompleteList=(Command="AudioDebugSound",Desc="<filter> Rejects new USoundBase requests where the sound name does not contain the provided filter")
-+ManualAutoCompleteList=(Command="AudioGetDynamicSoundVolume",Desc="Gets volume for given sound type ('Cue' or 'Wave') with provided name")
++ManualAutoCompleteList=(Command="AudioGetDynamicSoundVolume",Desc="Gets volume for given sound type ('Class', 'Cue' or 'Wave') with provided name")
 +ManualAutoCompleteList=(Command="AudioMemReport",Desc="Lists info for audio memory")
 +ManualAutoCompleteList=(Command="AudioMixerDebugSound",Desc="<filter> With new mixer enabled, rejects new USoundBase requests where the sound name does not contain the provided filter")
 +ManualAutoCompleteList=(Command="AudioResetAllDynamicSoundVolumes",Desc="Resets all dynamic volumes to unity")
-+ManualAutoCompleteList=(Command="AudioResetDynamicSoundVolume",Desc="Resets volume for given sound type ('Cue' or 'Wave') with provided name to unity")
-+ManualAutoCompleteList=(Command="AudioSetDynamicSoundVolume",Desc="Name=<name> Type=<type> Vol=<vol> Sets volume for given sound type ('Cue' or 'Wave') with provided name")
++ManualAutoCompleteList=(Command="AudioResetDynamicSoundVolume",Desc="Resets volume for given sound type ('Class', 'Cue' or 'Wave') with provided name to unity")
++ManualAutoCompleteList=(Command="AudioSetDynamicSoundVolume",Desc="Name=<name> Type=<type> Vol=<vol> Sets volume for given sound type ('Class', 'Cue' or 'Wave') with provided name")
 +ManualAutoCompleteList=(Command="AudioSoloSoundClass",Desc="<name> Solos USoundClass")
 +ManualAutoCompleteList=(Command="AudioSoloSoundCue",Desc="<name> Solos USoundCue")
 +ManualAutoCompleteList=(Command="AudioSoloSoundWave",Desc="<name> Solos USoundWave")
