Added redirects for some Pixel Streaming classes and properties which were renamed.
PixelStreamingInputComponent was renamed PixelStreamerInputComponent.
OnPixelStreamingInputEvent was renamed OnInputEvent.

diff --git a/Engine/Config/BaseEngine.ini b/Engine/Config/BaseEngine.ini
--- a/Engine/Config/BaseEngine.ini
+++ b/Engine/Config/BaseEngine.ini
@@ -737,6 +737,9 @@ bUseStreamingPause=false
 +PropertyRedirects=(OldName="AnimClassData.SubInstanceNodeProperties",NewName="LinkedAnimGraphNodeProperties")
 +PropertyRedirects=(OldName="AnimClassData.LayerNodeProperties",NewName="LinkedAnimLayerNodeProperties")
 
++ClassRedirects=(OldName="PixelStreamingInputComponent",NewName="PixelStreamerInputComponent")
++PropertyRedirects=(OldName="PixelStreamerInputComponent.OnPixelStreamingInputEvent",NewName="PixelStreamerInputComponent.OnInputEvent")
+
 [CoreUObject.Metadata]
 ; Note: UnrealHeaderTool should be rerun after making changes to MetadataRedirects to catch any keys specified in class headers
 +MetadataRedirects=(OldKey="K2Protected", NewKey="BlueprintProtected")
