Moving GizmoMaterial from editor to engine materials

diff --git a/Engine/Config/BaseEngine.ini b/Engine/Config/BaseEngine.ini
--- a/Engine/Config/BaseEngine.ini
+++ b/Engine/Config/BaseEngine.ini
@@ -124,7 +124,7 @@ DefaultVirtualTextureMaterialName=/Engine/EngineMaterials/DefaultVirtualTextureM
 TimecodeProviderClassName=None
 DefaultTimecodeProviderClassName=/Script/Engine.SystemTimeTimecodeProvider
 TextureStreamingBoundsMaterialName=/Engine/EditorMaterials/Utilities/TextureStreamingBounds_MATInst.TextureStreamingBounds_MATInst
-ArrowMaterialName=/Engine/EditorMaterials/GizmoMaterial.GizmoMaterial
+ArrowMaterialName=/Engine/EngineMaterials/GizmoMaterial.GizmoMaterial
 bAllowHostMigration=false
 HostMigrationTimeout=15
 ParticleEventManagerClassPath="/Script/Engine.ParticleEventManager"
