Add cine settings for hair rendeirng.

diff --git a/Engine/Config/BaseScalability.ini b/Engine/Config/BaseScalability.ini
--- a/Engine/Config/BaseScalability.ini
+++ b/Engine/Config/BaseScalability.ini
@@ -573,4 +573,7 @@ grass.DensityScale=1.0
 ;todo
 
 [ShadingQuality@Cine]
-;todo
+r.HairStrands.CoveragePass=1
+r.HairStrands.SkyLighting.SampleCount=64
+r.HairStrands.SkyLighting.JitterIntegration=1
+r.HairStrands.SkyAO.SampleCount=32
