Added to ListTextures console command help text

#rnx

diff --git a/Engine/Config/BaseInput.ini b/Engine/Config/BaseInput.ini
--- a/Engine/Config/BaseInput.ini
+++ b/Engine/Config/BaseInput.ini
@@ -116,7 +116,7 @@ MaxScrollbackSize=1024
 +ManualAutoCompleteList=(Command="Stat LIST",Desc="<Groups/Sets/Group> List groups of stats, saved sets, or specific stats within a specified group")
 +ManualAutoCompleteList=(Command="Stat splitscreen",Desc=)
 +ManualAutoCompleteList=(Command="MemReport",Desc="Outputs memory stats to a profile file. -Full gives more data, -Log outputs to the log")
-+ManualAutoCompleteList=(Command="ListTextures",Desc="Lists all loaded textures and their current memory footprint")
++ManualAutoCompleteList=(Command="ListTextures",Desc="[Streaming|NonStreaming|Forced] [-Alphasort] [-csv] Lists all loaded textures and their current memory footprint")
 +ManualAutoCompleteList=(Command="ListStreamingTextures",Desc="Lists info for all streaming textures")
 +ManualAutoCompleteList=(Command="ListAnims",Desc="Lists info for all animations")
 +ManualAutoCompleteList=(Command="ListSkeletalMeshes",Desc="Lists info for all skeletal meshes")
