Fix messages sent over steam sockets having invalid dormancy values which caused channels to close when they shouldn't.

diff --git a/Engine/Config/BaseEngine.ini b/Engine/Config/BaseEngine.ini
--- a/Engine/Config/BaseEngine.ini
+++ b/Engine/Config/BaseEngine.ini
@@ -1397,6 +1397,12 @@ ConnectionTimeout=80.0
 InitialConnectTimeout=120.0
 NetServerMaxTickRate=30
 MaxNetTickRate=120
+KeepAliveTime=0.2
+MaxClientRate=15000
+MaxInternetClientRate=10000
+RelevantTimeout=5.0
+SpawnPrioritySeconds=1.0
+ServerTravelPause=4.0
 
 [OnlineSubsystemMcp]
 bEnabled=false
