Updating suggested GPU drivers for Nvidia and AMD to 441.20 and 19.11.1 respectively. These were used during 4.24 GPU compat testing and did not reveal any driver-specific issues.

#rnx

diff --git a/Engine/Config/BaseHardware.ini b/Engine/Config/BaseHardware.ini
--- a/Engine/Config/BaseHardware.ini
+++ b/Engine/Config/BaseHardware.ini
@@ -32,7 +32,7 @@
 
 [GPU_NVIDIA]
 ; a newer one is out but has issues (see below)
-SuggestedDriverVersion="398.36"
+SuggestedDriverVersion="441.20"
 
 ; Seems this driver was not setting uninitialized global variables
 ; (static, non const) to 0 but leaving them uninitialized.
@@ -56,7 +56,7 @@ SuggestedDriverVersion="398.36"
 
 [GPU_AMD]
 ; From OR-27051 as of 8/4/2016 is 16.300.2311.0 7/18/2016 (Crimson Edition 16.9.2)
-SuggestedDriverVersion="16.60.2711.0"
+SuggestedDriverVersion="19.11.1"
 
 ; From OR-27051
 +Blacklist=(DriverVersion="==15.300.1025.1001", Reason="Crashes with Paragon content using more recent GPU features")
