1. When we changed the default resolution (Edit --> Editor Preferences --> Play ? Common Resolutions), the orientation (landscape or portrait) was reverted back to the original value for that default resolution (i.e., landscape for TVs, monitors, and laptops; and portrait for phones and tables). Fixed, it will now maintain the last orientation used. E.g., this allows to change multiple times from different phone/tablets devices in landscape mode without having to manually re-set the orientation.
2. Widget Blueprint --> Switch between Landscape and Portrait presents analog behaviour. In this part, the bug was fixed by forcing TVs, monitors, and laptops to have only landscape orientation. However, we might want them to be portrait (especially TVs or monitors). Thus, the code has been modified to present the exact same behavior than Edit --> Editor Preferences --> Play.
3. In addition, there was a bug in Widget Blueprint --> Switch between Landscape and Portrait. When the orientation was clicked right after opening the program (i.e., using the default values, before selecting a Screen Size), the screen would disappear and show a new resolution of 0x0. Bug fixed.
4. Additionally added some missing ôconstö keywords to some arguments in SDesignerView.h, SDesignerView.cpp, and LevelEditorPlaySettingsCustomization.h.

diff --git a/Engine/Config/BaseEditorPerProjectUserSettings.ini b/Engine/Config/BaseEditorPerProjectUserSettings.ini
--- a/Engine/Config/BaseEditorPerProjectUserSettings.ini
+++ b/Engine/Config/BaseEditorPerProjectUserSettings.ini
@@ -219,20 +219,20 @@ LastExecutedLaunchModeType=LaunchMode_OnDevice
 LastExecutedLaunchPlatform=
 
 ; common screen resolutions for laptops
-+LaptopScreenResolutions=(Description="Apple MacBook Air 11",Width=1366,Height=768,AspectRatio="16:9",bCanSwapAspectRatio=false)
-+LaptopScreenResolutions=(Description="Apple MacBook Air 13\"",Width=1440,Height=900,AspectRatio="16:10",bCanSwapAspectRatio=false)
-+LaptopScreenResolutions=(Description="Apple MacBook Pro 13\"",Width=1280,Height=800,AspectRatio="16:10",bCanSwapAspectRatio=false)
-+LaptopScreenResolutions=(Description="Apple MacBook Pro 13\" (Retina)",Width=2560,Height=1600,AspectRatio="16:10",bCanSwapAspectRatio=false)
-+LaptopScreenResolutions=(Description="Apple MacBook Pro 15\"",Width=1440,Height=900,AspectRatio="16:10",bCanSwapAspectRatio=false)
-+LaptopScreenResolutions=(Description="Apple MacBook Pro 15\" (Retina)",Width=2880,Height=1800,AspectRatio="16:10",bCanSwapAspectRatio=false)
-+LaptopScreenResolutions=(Description="Generic 14-15.6\" Notebook",Width=1366,Height=768,AspectRatio="16:9",bCanSwapAspectRatio=false)
++LaptopScreenResolutions=(Description="Apple MacBook Air 11",Width=1366,Height=768,AspectRatio="16:9",bCanSwapAspectRatio=true)
++LaptopScreenResolutions=(Description="Apple MacBook Air 13\"",Width=1440,Height=900,AspectRatio="16:10",bCanSwapAspectRatio=true)
++LaptopScreenResolutions=(Description="Apple MacBook Pro 13\"",Width=1280,Height=800,AspectRatio="16:10",bCanSwapAspectRatio=true)
++LaptopScreenResolutions=(Description="Apple MacBook Pro 13\" (Retina)",Width=2560,Height=1600,AspectRatio="16:10",bCanSwapAspectRatio=true)
++LaptopScreenResolutions=(Description="Apple MacBook Pro 15\"",Width=1440,Height=900,AspectRatio="16:10",bCanSwapAspectRatio=true)
++LaptopScreenResolutions=(Description="Apple MacBook Pro 15\" (Retina)",Width=2880,Height=1800,AspectRatio="16:10",bCanSwapAspectRatio=true)
++LaptopScreenResolutions=(Description="Generic 14-15.6\" Notebook",Width=1366,Height=768,AspectRatio="16:9",bCanSwapAspectRatio=true)
 
 ; common screen resolutions for desktop monitors
-+MonitorScreenResolutions=(Description="19\" monitor",Width=1440,Height=900,AspectRatio="16:10",bCanSwapAspectRatio=false)
-+MonitorScreenResolutions=(Description="20\" monitor",Width=1600,Height=900,AspectRatio="16:9",bCanSwapAspectRatio=false)
-+MonitorScreenResolutions=(Description="22\" monitor",Width=1680,Height=1050,AspectRatio="16:10",bCanSwapAspectRatio=false)
-+MonitorScreenResolutions=(Description="21.5-24\" monitor",Width=1920,Height=1080,AspectRatio="16:9",bCanSwapAspectRatio=false)
-+MonitorScreenResolutions=(Description="27\" monitor",Width=2560,Height=1440,AspectRatio="16:9",bCanSwapAspectRatio=false)
++MonitorScreenResolutions=(Description="19\" monitor",Width=1440,Height=900,AspectRatio="16:10",bCanSwapAspectRatio=true)
++MonitorScreenResolutions=(Description="20\" monitor",Width=1600,Height=900,AspectRatio="16:9",bCanSwapAspectRatio=true)
++MonitorScreenResolutions=(Description="22\" monitor",Width=1680,Height=1050,AspectRatio="16:10",bCanSwapAspectRatio=true)
++MonitorScreenResolutions=(Description="21.5-24\" monitor",Width=1920,Height=1080,AspectRatio="16:9",bCanSwapAspectRatio=true)
++MonitorScreenResolutions=(Description="27\" monitor",Width=2560,Height=1440,AspectRatio="16:9",bCanSwapAspectRatio=true)
 
 ; common screen resolutions for mobile phones
 +PhoneScreenResolutions=(Description="Apple iPhone 5S",Width=320,Height=568,AspectRatio="~16:9",bCanSwapAspectRatio=true,ProfileName="iPhone5S")
@@ -285,10 +285,10 @@ LastExecutedLaunchPlatform=
 +TabletScreenResolutions=(Description="Microsoft Surface Pro",Width=1080,Height=1920,AspectRatio="9:16",bCanSwapAspectRatio=true)
 
 ; common screen resolutions for television sets
-+TelevisionScreenResolutions=(Description="720p (HDTV, Blu-ray)",Width=1280,Height=720,AspectRatio="16:9",bCanSwapAspectRatio=false)
-+TelevisionScreenResolutions=(Description="1080i, 1080p (HDTV, Blu-ray)",Width=1920,Height=1080,AspectRatio="16:9",bCanSwapAspectRatio=false)
-+TelevisionScreenResolutions=(Description="4K Ultra HD",Width=3840,Height=2160,AspectRatio="16:9",bCanSwapAspectRatio=false)
-+TelevisionScreenResolutions=(Description="4K Digital Cinema",Width=4096,Height=2160,AspectRatio="1.90:1",bCanSwapAspectRatio=false)
++TelevisionScreenResolutions=(Description="720p (HDTV, Blu-ray)",Width=1280,Height=720,AspectRatio="16:9",bCanSwapAspectRatio=true)
++TelevisionScreenResolutions=(Description="1080i, 1080p (HDTV, Blu-ray)",Width=1920,Height=1080,AspectRatio="16:9",bCanSwapAspectRatio=true)
++TelevisionScreenResolutions=(Description="4K Ultra HD",Width=3840,Height=2160,AspectRatio="16:9",bCanSwapAspectRatio=true)
++TelevisionScreenResolutions=(Description="4K Digital Cinema",Width=4096,Height=2160,AspectRatio="1.90:1",bCanSwapAspectRatio=true)
 
 
 [/Script/UnrealEd.LevelEditorViewportSettings]
diff --git a/Engine/Source/Editor/UMGEditor/Private/Designer/SDesignerView.cpp b/Engine/Source/Editor/UMGEditor/Private/Designer/SDesignerView.cpp
--- a/Engine/Source/Editor/UMGEditor/Private/Designer/SDesignerView.cpp
+++ b/Engine/Source/Editor/UMGEditor/Private/Designer/SDesignerView.cpp
@@ -335,11 +335,11 @@ void SDesignerView::Construct(const FArguments& InArgs, TSharedPtr<FWidgetBluepr
 
 	bShowResolutionOutlines = false;
 
+	HeightReadFromSettings = 0;
+	WidthReadFromSettings = 0;
 	SetStartupResolution();
 
 	CachedPreviewDesiredSize = FVector2D(0, 0);
-	HeightReadFromSettings = 0;
-	WidthReadFromSettings = 0;
 
 	ResolutionTextFade = FCurveSequence(0.0f, 1.0f);
 	ResolutionTextFade.Play(this->AsShared());
@@ -961,12 +961,16 @@ void SDesignerView::SetStartupResolution()
 		GConfig->SetInt(*ConfigSectionName, TEXT("PreviewWidth"), DefaultResolutionWidth, GEditorPerProjectIni);
 		PreviewWidth = DefaultResolutionWidth;
 	}
+	// Initially assign WidthReadFromSettings to PreviewWidth
+	WidthReadFromSettings = PreviewWidth;
 	// Height
 	if (!GConfig->GetInt(*ConfigSectionName, TEXT("PreviewHeight"), PreviewHeight, GEditorPerProjectIni))
 	{
 		GConfig->SetInt(*ConfigSectionName, TEXT("PreviewHeight"), DefaultResolutionHeight, GEditorPerProjectIni);
 		PreviewHeight = DefaultResolutionHeight;
 	}
+	// Initially assign HeightReadFromSettings to PreviewHeight
+	HeightReadFromSettings = PreviewHeight;
 	// Aspect Ratio
 	if (!GConfig->GetString(*ConfigSectionName, TEXT("PreviewAspectRatio"), PreviewAspectRatio, GEditorPerProjectIni))
 	{
@@ -3254,15 +3258,14 @@ FReply SDesignerView::HandleDPISettingsClicked()
 	return FReply::Handled();
 }
 
-void SDesignerView::HandleOnCommonResolutionSelected(FPlayScreenResolution InResolution)
+void SDesignerView::HandleOnCommonResolutionSelected(const FPlayScreenResolution InResolution)
 {
 	bSafeZoneFlipped = false;
 	bCanPreviewSwapAspectRatio = InResolution.bCanSwapAspectRatio;
 	WidthReadFromSettings = InResolution.Width;
 	HeightReadFromSettings = InResolution.Height;
-	// Phone/tablet resolutions can be stored in either portrait or landscape mode, and may need to be flipped
-	if (bCanPreviewSwapAspectRatio && ((!bPreviewIsPortrait && InResolution.Width < InResolution.Height) ||
-		(bPreviewIsPortrait && InResolution.Width > InResolution.Height)))
+	// Most resolutions (tablets, phones, TVs, etc.) can be stored in either portrait or landscape mode, and may need to be flipped
+	if (bCanPreviewSwapAspectRatio && (bPreviewIsPortrait != (InResolution.Width < InResolution.Height)))
 	{
 		PreviewWidth = InResolution.Height;
 		PreviewHeight = InResolution.Width;
@@ -3271,17 +3274,12 @@ void SDesignerView::HandleOnCommonResolutionSelected(FPlayScreenResolution InRes
 	{
 		PreviewWidth = InResolution.Width;
 		PreviewHeight = InResolution.Height;
-		bPreviewIsPortrait = PreviewWidth < PreviewHeight;
 	}
+	bPreviewIsPortrait = PreviewWidth < PreviewHeight;
 	PreviewAspectRatio = InResolution.AspectRatio;
 
 	PreviewOverrideName = InResolution.ProfileName;
 
-	if (!bCanPreviewSwapAspectRatio)
-	{
-		bPreviewIsPortrait = (PreviewHeight > PreviewWidth);
-	}
-
 	ScaleFactor = 1.0f;
 	ULevelEditorPlaySettings* PlayInSettings = GetMutableDefault<ULevelEditorPlaySettings>();
 	const UDeviceProfile* DeviceProfile = UDeviceProfileManager::Get().FindProfile(PreviewOverrideName, false);
@@ -3336,7 +3334,7 @@ void SDesignerView::HandleOnCommonResolutionSelected(FPlayScreenResolution InRes
 	ResolutionTextFade.Play(this->AsShared());
 }
 
-bool SDesignerView::HandleIsCommonResolutionSelected(FPlayScreenResolution InResolution) const
+bool SDesignerView::HandleIsCommonResolutionSelected(const FPlayScreenResolution InResolution) const
 {
 	// If we're using a custom design time size, none of the other resolutions should appear selected, even if they match.
 	if ( UUserWidget* DefaultWidget = GetDefaultWidget() )
@@ -3616,8 +3614,8 @@ FReply SDesignerView::HandleZoomToFitClicked()
 FReply SDesignerView::HandleSwapAspectRatioClicked()
 {
 	bSafeZoneFlipped = false;
-	// If in portrait
-	if (PreviewHeight > PreviewWidth)
+	// If in default orientation (portrait for table/phone, landscape for monitor/laptop/TV)
+	if ((WidthReadFromSettings < HeightReadFromSettings) == (PreviewWidth < PreviewHeight))
 	{
 		PreviewHeight = WidthReadFromSettings;
 		PreviewWidth = HeightReadFromSettings;
diff --git a/Engine/Source/Editor/UMGEditor/Private/Designer/SDesignerView.h b/Engine/Source/Editor/UMGEditor/Private/Designer/SDesignerView.h
--- a/Engine/Source/Editor/UMGEditor/Private/Designer/SDesignerView.h
+++ b/Engine/Source/Editor/UMGEditor/Private/Designer/SDesignerView.h
@@ -195,8 +195,8 @@ private:
 	EVisibility GetCustomResolutionEntryVisibility() const;
 	
 	// Handles selecting a common screen resolution.
-	void HandleOnCommonResolutionSelected(FPlayScreenResolution InResolution);
-	bool HandleIsCommonResolutionSelected(FPlayScreenResolution InResolution) const;
+	void HandleOnCommonResolutionSelected(const FPlayScreenResolution InResolution);
+	bool HandleIsCommonResolutionSelected(const FPlayScreenResolution InResolution) const;
 	void AddScreenResolutionSection(FMenuBuilder& MenuBuilder, const TArray<FPlayScreenResolution> Resolutions, const FText SectionName);
 	TSharedRef<SWidget> GetResolutionsMenu();
 
diff --git a/Engine/Source/Editor/UnrealEd/Private/Settings/LevelEditorPlaySettingsCustomization.h b/Engine/Source/Editor/UnrealEd/Private/Settings/LevelEditorPlaySettingsCustomization.h
--- a/Engine/Source/Editor/UnrealEd/Private/Settings/LevelEditorPlaySettingsCustomization.h
+++ b/Engine/Source/Editor/UnrealEd/Private/Settings/LevelEditorPlaySettingsCustomization.h
@@ -380,12 +380,24 @@ protected:
 private:
 
 	// Handles selecting a common screen resolution.
-	void HandleCommonResolutionSelected(FPlayScreenResolution Resolution)
+	void HandleCommonResolutionSelected(const FPlayScreenResolution Resolution)
 	{
 		int32 Width = Resolution.Width;
 		int32 Height = Resolution.Height;
 		float ScaleFactor;
 		ULevelEditorPlaySettings* PlayInSettings = GetMutableDefault<ULevelEditorPlaySettings>();
+		// Maintain previous orientation (i.e., swap Width and Height if required)
+		int32 PreviousWidth;
+		int32 PreviousHeight;
+		if (WindowWidthProperty->GetValue(PreviousWidth) == FPropertyAccess::Success && WindowHeightProperty->GetValue(PreviousHeight) == FPropertyAccess::Success)
+		{
+			const bool bIsOrientationPreserved = (PreviousWidth < PreviousHeight) != (Width < Height);
+			if (bIsOrientationPreserved)
+			{
+				Width = Resolution.Height;
+				Height = Resolution.Width;
+			}
+		}
 
 		UDeviceProfile* DeviceProfile = UDeviceProfileManager::Get().FindProfile(Resolution.ProfileName, false);
 		if (DeviceProfile)
