Sequencer: Make template sequences restore state when finished by default.

diff --git a/Engine/Config/BaseEngine.ini b/Engine/Config/BaseEngine.ini
--- a/Engine/Config/BaseEngine.ini
+++ b/Engine/Config/BaseEngine.ini
@@ -2331,6 +2331,9 @@ WatchdogMinimumHangSeconds=120
 [/Script/LevelSequence.LevelSequence]
 DefaultCompletionMode=RestoreState
 
+[/Script/TemplateSequence.TemplateSequence]
+DefaultCompletionMode=RestoreState
+
 [PlatformCrypto]
 PlatformRequiresDataCrypto=False
 PakSigningRequired=True
