Disabling the packed outputs optimization temporarily due to a bug.

diff --git a/Engine/Config/BaseEngine.ini b/Engine/Config/BaseEngine.ini
--- a/Engine/Config/BaseEngine.ini
+++ b/Engine/Config/BaseEngine.ini
@@ -1357,6 +1357,8 @@ con.DebugLateCheat = True
 ; disabled for now to prevent issues http://crashreport/Crashes/Show/2948912, http://crashreport/Crashes/Show/2948842
 r.DetectAndWarnOfBadDrivers = 0
 
+vm.BatchPackedVMOutput=0
+
 [SystemSettingsEditor]
 ; System settings overrides for the editor.  Ideally the editor should use the same settings as the game.
 
