#UE4 Add default submix paths to BaseEngine.ini so they will be cooked

diff --git a/Engine/Config/BaseEngine.ini b/Engine/Config/BaseEngine.ini
--- a/Engine/Config/BaseEngine.ini
+++ b/Engine/Config/BaseEngine.ini
@@ -938,6 +938,10 @@ BackRight=7
 DefaultSoundClassName=/Engine/EngineSounds/Master.Master
 DefaultMediaSoundClassName=/Engine/EngineSounds/Master.Master
 DefaultSoundSubmixName=/Engine/EngineSounds/MasterSubmix.MasterSubmix
+MasterSubmix=/Engine/EngineSounds/Submixes/MasterSubmixDefault.MasterSubmixDefault
+ReverbSubmix=/Engine/EngineSounds/Submixes/MasterReverbSubmixDefault.MasterReverbSubmixDefault
+EQSubmix=/Engine/EngineSounds/Submixes/MasterEQSubmixDefault.MasterEQSubmixDefault
+AmbisonicSubmix=/Engine/EngineSounds/Submixes/MasterAmbisonicSubmixDefault.MasterAmbisonicSubmixDefault
 LowPassFilterResonance=0.9
 MaximumConcurrentStreams=2
 DialogueFilenameFormat="{DialogueGuid}_{ContextId}"
