Blacklist some older AMD drivers that are contributing to swapchain creation crashes on various laptops

diff --git a/Engine/Config/BaseHardware.ini b/Engine/Config/BaseHardware.ini
--- a/Engine/Config/BaseHardware.ini
+++ b/Engine/Config/BaseHardware.ini
@@ -61,8 +61,8 @@ SuggestedDriverVersion="19.11.1"
 ; From OR-27051
 +Blacklist=(DriverVersion="==15.300.1025.1001", Reason="Crashes with Paragon content using more recent GPU features")
 
-; 14.301.1001.0 - 13.152.1.1000 are the AMD drivers that seem to be causing crashes when we try to compile a compute shader.
-+Blacklist=(DriverVersion="<=14.301.1001.0", Reason="Crashes with Paragon content using more recent GPU features")
+; 2017 drivers and older contribute heavily to swapchain creation errors.
++Blacklist=(DriverVersion="<=22.19.662.4", Reason="Older drivers known to cause crashes with integrated laptop graphics on at least R5 and R6")
 
 ;-----------------------------------------------------------------------------------------------------------------
 
