Update default target SDK to 28 for Android since Google Play now requires it

#ue4
#android

diff --git a/Engine/Config/BaseEngine.ini b/Engine/Config/BaseEngine.ini
--- a/Engine/Config/BaseEngine.ini
+++ b/Engine/Config/BaseEngine.ini
@@ -2051,7 +2051,7 @@ StoreVersionOffsetX86=0
 StoreVersionOffsetX8664=0
 VersionDisplayName=1.0
 MinSDKVersion=9
-TargetSDKVersion=26
+TargetSDKVersion=28
 bEnableGradle=true
 bEnableLint=false
 bShowLaunchImage=true
