Add missing timeout values for SteamSockets. This fixes the base interface booting off users when it shouldn't.

#rnx

diff --git a/Engine/Config/BaseEngine.ini b/Engine/Config/BaseEngine.ini
--- a/Engine/Config/BaseEngine.ini
+++ b/Engine/Config/BaseEngine.ini
@@ -1384,6 +1384,8 @@ NetConnectionClassName="/Script/OnlineSubsystemSteam.SteamNetConnection"
 
 [/Script/SteamSockets.SteamSocketsNetDriver]
 NetConnectionClassName="/Script/SteamSockets.SteamSocketsNetConnection"
+ConnectionTimeout=80.0
+InitialConnectTimeout=120.0
 
 [OnlineSubsystemMcp]
 bEnabled=false
