Added PhysicalMaterialMask thumbnail renderer.

diff --git a/Engine/Config/BaseEditor.ini b/Engine/Config/BaseEditor.ini
--- a/Engine/Config/BaseEditor.ini
+++ b/Engine/Config/BaseEditor.ini
@@ -34,6 +34,7 @@ ColorFontAtlasSize=512
 +RenderableThumbnailTypes=(ClassNeedingThumbnailName="/Script/Engine.MaterialInterface",RendererClassName="/Script/UnrealEd.MaterialInstanceThumbnailRenderer")
 +RenderableThumbnailTypes=(ClassNeedingThumbnailName="/Script/Engine.MaterialFunctionInterface",RendererClassName="/Script/UnrealEd.MaterialFunctionThumbnailRenderer")
 +RenderableThumbnailTypes=(ClassNeedingThumbnailName="/Script/Engine.ParticleSystem",RendererClassName="/Script/UnrealEd.ParticleSystemThumbnailRenderer")
++RenderableThumbnailTypes=(ClassNeedingThumbnailName="/Script/Engine.PhysicalMaterialMask",RendererClassName="/Script/UnrealEd.PhysicalMaterialMaskThumbnailRenderer")
 +RenderableThumbnailTypes=(ClassNeedingThumbnailName="/Script/Engine.StaticMesh",RendererClassName="/Script/UnrealEd.StaticMeshThumbnailRenderer")
 +RenderableThumbnailTypes=(ClassNeedingThumbnailName="/Script/Engine.SkeletalMesh",RendererClassName="/Script/UnrealEd.SkeletalMeshThumbnailRenderer")
 +RenderableThumbnailTypes=(ClassNeedingThumbnailName="/Script/Engine.AnimSequenceBase",RendererClassName="/Script/UnrealEd.AnimSequenceThumbnailRenderer")
