UE4 - Blacklist D3D12 driver
#rnx

diff --git a/Engine/Config/BaseHardware.ini b/Engine/Config/BaseHardware.ini
--- a/Engine/Config/BaseHardware.ini
+++ b/Engine/Config/BaseHardware.ini
@@ -22,7 +22,8 @@
 ;    [GPU_NVIDIA] [GPU_AMD] [GPU_Intel]
 ;
 ; Explanation of the INI entries:
-;  - "SuggestedDriverVersion": string used to communicate to the user which driver version is appropriate (not used by any code), empty if unknown
+;  - "SuggestedDriverVersion": string used to communicate to the user which driver version is appropriate (not used by any code), empty if unknown.
+;		Can use multiple so we can add per RHI appending ";D3D12" for instance for a specific RHI; having no ';RHI' will return that one if the RHI is not found
 ;  - "Blacklist": array of tests that can be made to mark a driver or a range of them as problematic
 ;    "DriverVersion": to compare against the "unified driver" (see log) version e.g. "123.45.2" "==123.45" "<123.45", comparison operators: == != > < >= <=
 ;    "Reason": string used to document the reason for blacklisting the driver, could be shown to the user
@@ -57,6 +58,7 @@ SuggestedDriverVersion="441.20"
 [GPU_AMD]
 ; From OR-27051 as of 8/4/2016 is 16.300.2311.0 7/18/2016 (Crimson Edition 16.9.2)
 SuggestedDriverVersion="19.11.1"
++SuggestedDriverVersion="19.20.1;D3D12"
 
 ; From OR-27051
 +Blacklist=(DriverVersion="==15.300.1025.1001", Reason="Crashes with Paragon content using more recent GPU features")
@@ -64,6 +66,8 @@ SuggestedDriverVersion="19.11.1"
 ; 2017 drivers and older contribute heavily to swapchain creation errors.
 +Blacklist=(DriverVersion="<=22.19.662.4", Reason="Older drivers known to cause crashes with integrated laptop graphics on at least R5 and R6")
 
++Blacklist=(DriverVersion="==26.20.13031.15006", Reason="Driver crashes creating PSOs", RHI="D3D12")
+
 ;-----------------------------------------------------------------------------------------------------------------
 
 [GPU_Intel]
