Device profile cleanup
- Remove unsupported iOS and Android devices
- Map Android devices from minspec (Adreno 4xx/Mali T6xx) and above to Android_Low, Android_Mid, Android_High and then to Scalablity Groups 0,1,2 respectively.
- Disable Vulkan by default, enabling only on Adreno 6xx Android 9+, Adreno G72/G76 and PowerVR 9xxx devices
- Map range of iOS devices to post process and shadow quality using scalability groups, and all settings shared

diff --git a/Engine/Config/Android/AndroidScalability.ini b/Engine/Config/Android/AndroidScalability.ini
--- a/Engine/Config/Android/AndroidScalability.ini
+++ b/Engine/Config/Android/AndroidScalability.ini
@@ -1,3 +1,110 @@
+; ------------------------------------------------------------------------------
+; Foliage settings
+; ------------------------------------------------------------------------------
+
+[FoliageQuality@0]
+foliage.DensityScale=0.5
+grass.DensityScale=0.5
+
+[FoliageQuality@1]
+grass.densityScale=0.8
+grass.CullDistanceScale=0.8
+foliage.DensityScale=0.8
+
+[FoliageQuality@2]
+grass.densityScale=1.0
+grass.CullDistanceScale=1.0
+foliage.DensityScale=1.0
+
+[FoliageQuality@3]
+grass.densityScale=1.0
+grass.CullDistanceScale=1.0
+foliage.DensityScale=1.0
+
+; ------------------------------------------------------------------------------
+; View distance settings
+; ------------------------------------------------------------------------------
+
+[ViewDistanceQuality@0]
+r.ViewDistanceScale=0.70
+
+[ViewDistanceQuality@1]
+r.ViewDistanceScale=1.0
+
+[ViewDistanceQuality@2]
+r.ViewDistanceScale=1.0
+
+[ViewDistanceQuality@3]
+r.ViewDistanceScale=1.0
+
+; ------------------------------------------------------------------------------
+; Shadow settings
+; ------------------------------------------------------------------------------
+
+[ShadowQuality@0]
+r.ShadowQuality=0
+r.Shadow.MaxCSMResolution=1024
+r.Shadow.DistanceScale=0.5
+r.Shadow.CSM.MaxCascades=1
+r.Shadow.RadiusThreshold=0.06
+
+[ShadowQuality@1]
+r.ShadowQuality=1
+r.Shadow.MaxCSMResolution=1024
+r.Shadow.DistanceScale=0.7
+r.Shadow.CSM.MaxCascades=1
+r.Shadow.RadiusThreshold=0.06
+
+[ShadowQuality@2]
+r.ShadowQuality=1
+r.Shadow.MaxCSMResolution=2048
+r.Shadow.DistanceScale=1.0
+r.Shadow.CSM.MaxCascades=2
+r.Shadow.RadiusThreshold=0.06
+
+[ShadowQuality@3]
+r.ShadowQuality=2
+r.Shadow.MaxCSMResolution=2048
+r.Shadow.DistanceScale=1.0
+r.Shadow.CSM.MaxCascades=2
+r.Shadow.RadiusThreshold=0.06
+
+; ------------------------------------------------------------------------------
+; Postprocess settings
+; ------------------------------------------------------------------------------
+
+[PostProcessQuality@0]
+r.Tonemapper.Quality=1
+r.BloomQuality=0
+r.DepthOfFieldQuality=0
+r.LightShaftQuality=0
+r.RefractionQuality=0
+Slate.ForceBackgroundBlurLowQualityOverride=1
+
+[PostProcessQuality@1]
+r.Tonemapper.Quality=2
+r.BloomQuality=1
+r.DepthOfFieldQuality=1
+r.LightShaftQuality=1
+r.RefractionQuality=1
+Slate.ForceBackgroundBlurLowQualityOverride=1
+
+[PostProcessQuality@2]
+r.Tonemapper.Quality=3
+r.BloomQuality=1
+r.DepthOfFieldQuality=1
+r.LightShaftQuality=1
+r.RefractionQuality=1
+Slate.ForceBackgroundBlurLowQualityOverride=1
+
+[PostProcessQuality@3]
+r.Tonemapper.Quality=3
+r.BloomQuality=1
+r.DepthOfFieldQuality=1
+r.LightShaftQuality=1
+r.RefractionQuality=1
+Slate.ForceBackgroundBlurLowQualityOverride=1
+
 ; ------------------------------------------------------------------------------
 ; Effects settings
 ; ------------------------------------------------------------------------------
@@ -45,3 +152,27 @@ r.SkyAtmosphere.FastSkyLUT.SampleCountMax=8.0
 r.SkyAtmosphere.FastSkyLUT.Width=96
 r.SkyAtmosphere.FastSkyLUT.Height=50
 r.SkyAtmosphere.TransmittanceLUT.UseSmallFormat=1
+
+; ------------------------------------------------------------------------------
+; Texture settings
+; ------------------------------------------------------------------------------
+
+[TextureQuality@0]
+r.MaxAnisotropy=1
+
+[TextureQuality@1]
+r.MaxAnisotropy=1
+
+[TextureQuality@2]
+r.MaxAnisotropy=2
+
+[TextureQuality@3]
+r.MaxAnisotropy=4
+
+[ScalabilityPreview]
+; don't want to preview any texture streaming changes in the editor
+BlacklistCVars=r.MaxAnisotropy
++BlacklistCVars=r.Streaming.Boost
++BlacklistCVars=r.Streaming.AmortizeCPUToGPUCopy
++BlacklistCVars=r.Streaming.MaxNumTexturesToStreamPerFrame
++BlacklistCVars=r.Streaming.PoolSize
\ No newline at end of file
diff --git a/Engine/Config/BaseDeviceProfiles.ini b/Engine/Config/BaseDeviceProfiles.ini
--- a/Engine/Config/BaseDeviceProfiles.ini
+++ b/Engine/Config/BaseDeviceProfiles.ini
@@ -9,19 +9,14 @@
 +DeviceProfileNameAndTypes=WindowsServer,Windows
 +DeviceProfileNameAndTypes=WindowsClient,Windows
 +DeviceProfileNameAndTypes=IOS,IOS
-+DeviceProfileNameAndTypes=iPad2,IOS
-+DeviceProfileNameAndTypes=iPad3,IOS
-+DeviceProfileNameAndTypes=iPad4,IOS
 +DeviceProfileNameAndTypes=iPadAir,IOS
-+DeviceProfileNameAndTypes=iPadMini,IOS
++DeviceProfileNameAndTypes=iPadAir2,IOS
++DeviceProfileNameAndTypes=iPadAir3,IOS
 +DeviceProfileNameAndTypes=iPadMini2,IOS
 +DeviceProfileNameAndTypes=iPadMini4,IOS
-+DeviceProfileNameAndTypes=iPodTouch5,IOS
++DeviceProfileNameAndTypes=iPadMini5,IOS
 +DeviceProfileNameAndTypes=iPodTouch6,IOS
 +DeviceProfileNameAndTypes=iPodTouch7,IOS
-+DeviceProfileNameAndTypes=iPhone4,IOS
-+DeviceProfileNameAndTypes=iPhone4S,IOS
-+DeviceProfileNameAndTypes=iPhone5,IOS
 +DeviceProfileNameAndTypes=iPhone5S,IOS
 +DeviceProfileNameAndTypes=iPhone6,IOS
 +DeviceProfileNameAndTypes=iPhone6Plus,IOS
@@ -49,8 +44,6 @@
 +DeviceProfileNameAndTypes=iPad7,IOS
 +DeviceProfileNameAndTypes=iPadPro11,IOS
 +DeviceProfileNameAndTypes=iPadPro3_129,IOS
-+DeviceProfileNameAndTypes=iPadAir3,IOS
-+DeviceProfileNameAndTypes=iPadMini5,IOS
 +DeviceProfileNameAndTypes=AppleTV,IOS
 +DeviceProfileNameAndTypes=AppleTV4K,IOS
 +DeviceProfileNameAndTypes=Mac,Mac
@@ -69,32 +62,25 @@
 +DeviceProfileNameAndTypes=Android_Mid,Android
 +DeviceProfileNameAndTypes=Android_High,Android
 +DeviceProfileNameAndTypes=Android_Default,Android
-+DeviceProfileNameAndTypes=Android_Adreno320,Android
-+DeviceProfileNameAndTypes=Android_Adreno2xx,Android
-+DeviceProfileNameAndTypes=Android_Adreno330,Android
-+DeviceProfileNameAndTypes=Android_Adreno330_Ver53,Android
-+DeviceProfileNameAndTypes=Android_Adreno4xx_Nexus6,Android
 +DeviceProfileNameAndTypes=Android_Adreno4xx,Android
++DeviceProfileNameAndTypes=Android_Adreno5xx_Low,Android
 +DeviceProfileNameAndTypes=Android_Adreno5xx,Android
-+DeviceProfileNameAndTypes=Android_Adreno5xx_No_Vulkan,Android
 +DeviceProfileNameAndTypes=Android_Adreno6xx,Android
-+DeviceProfileNameAndTypes=Android_PowerVR54x,Android
-+DeviceProfileNameAndTypes=Android_PowerVR_G6xxx,Android
-+DeviceProfileNameAndTypes=Android_PowerVRHigh,Android
-+DeviceProfileNameAndTypes=Android_PowerVRLow,Android
-+DeviceProfileNameAndTypes=Android_PowerVRMid,Android
-+DeviceProfileNameAndTypes=Android_TegraK1,Android
-+DeviceProfileNameAndTypes=Android_Tegra4,Android
-+DeviceProfileNameAndTypes=Android_IntelHD,Android
-+DeviceProfileNameAndTypes=Android_IntelHD_ES3,Android
-+DeviceProfileNameAndTypes=Android_Mali_4xx,Android
++DeviceProfileNameAndTypes=Android_Adreno6xx_Vulkan,Android
 +DeviceProfileNameAndTypes=Android_Mali_T6xx,Android
 +DeviceProfileNameAndTypes=Android_Mali_T7xx,Android
 +DeviceProfileNameAndTypes=Android_Mali_T8xx,Android
-+DeviceProfileNameAndTypes=Android_Mali_T8xx_No_Vulkan,Android
 +DeviceProfileNameAndTypes=Android_Mali_G71,Android
 +DeviceProfileNameAndTypes=Android_Mali_G72,Android
++DeviceProfileNameAndTypes=Android_Mali_G72_Vulkan,Android
 +DeviceProfileNameAndTypes=Android_Mali_G76,Android
++DeviceProfileNameAndTypes=Android_Mali_G76_Vulkan,Android
++DeviceProfileNameAndTypes=Android_PowerVR_G6xxx,Android
++DeviceProfileNameAndTypes=Android_PowerVR_GT7xxx,Android
++DeviceProfileNameAndTypes=Android_PowerVR_GE8xxx,Android
++DeviceProfileNameAndTypes=Android_PowerVR_GM9xxx,Android
++DeviceProfileNameAndTypes=Android_PowerVR_GM9xxx_Vulkan,Android
++DeviceProfileNameAndTypes=Android_TegraK1,Android
 +DeviceProfileNameAndTypes=Lumin,Lumin
 +DeviceProfileNameAndTypes=Lumin_Desktop,Lumin
 
@@ -149,14 +135,12 @@ BaseProfileName=Windows
 DeviceType=Windows
 BaseProfileName=Windows
 
+;----------------
+; iOS Devices
+
 [IOS DeviceProfile]
 DeviceType=IOS
 BaseProfileName=
-+CVars=r.BloomQuality=0
-+CVars=r.DepthOfFieldQuality=0
-+CVars=r.LightShaftQuality=0
-+CVars=r.RefractionQuality=0
-+CVars=r.ShadowQuality=2
 ; HZB will be slower with tiled, and not needed
 +CVars=r.HZBOcclusion=0
 ; Non need for depth prepass
@@ -165,12 +149,18 @@ BaseProfileName=
 +CVars=r.TranslucentLightingVolume=0
 ; Needs geometry shader support
 +CVars=r.AllowPointLightCubemapShadows=0
-; Temporal AA (?) looks blurry right now
-+CVars=r.PostProcessAAQuality=0
 ; Not sure on this one - may need the whole D Buffer stuff going
 +CVars=r.Decal.StencilSizeThreshold=-1
 +CVars=slate.AbsoluteIndices=1
 +CVars=r.MorphTarget.Mode=0
+; Scalability
++CVars=sg.AntiAliasingQuality=0
++CVars=sg.ShadowQuality=1
++CVars=sg.ViewDistanceQuality=3
++CVars=sg.PostProcessQuality=3
++CVars=sg.TextureQuality=3
++CVars=sg.EffectsQuality=3
++CVars=sg.FoliageQuality=3
 
 [iPad2 DeviceProfile]
 DeviceType=IOS
@@ -190,15 +180,13 @@ BaseProfileName=IOS
 [iPadAir DeviceProfile]
 DeviceType=IOS
 BaseProfileName=IOS
-+CVars=r.BloomQuality=1
++CVars=sg.PostProcessQuality=1
 
 [iPadAir2 DeviceProfile]
 DeviceType=IOS
 BaseProfileName=IOS
-+CVars=r.BloomQuality=1
-+CVars=r.DepthOfFieldQuality=1
-+CVars=r.LightShaftQuality=1
-+CVars=r.RefractionQuality=1
++CVars=sg.PostProcessQuality=3
++CVars=sg.ShadowQuality=2
 
 [iPadMini DeviceProfile]
 DeviceType=IOS
@@ -212,36 +200,13 @@ BaseProfileName=iPadAir
 [iPadMini4 DeviceProfile]
 DeviceType=IOS
 BaseProfileName=iPadMini2
-+CVars=r.RefractionQuality=1
-
-[iPhone4 DeviceProfile]
-DeviceType=IOS
-BaseProfileName=IOS
-+CVars=r.RenderTargetSwitchWorkaround=1
-
-[iPhone4S DeviceProfile]
-DeviceType=IOS
-BaseProfileName=IOS
-+CVars=r.RenderTargetSwitchWorkaround=1
-
-[iPhone5 DeviceProfile]
-DeviceType=IOS
-BaseProfileName=IOS
-+CVars=r.RenderTargetSwitchWorkaround=1
++CVars=sg.PostProcessQuality=2
 
 [iPhone5S DeviceProfile]
 DeviceType=IOS
 BaseProfileName=IOS
 +CVars=r.MobileContentScaleFactor=2
-+CVars=r.BloomQuality=1
-+CVars=r.DepthOfFieldQuality=1
-+CVars=r.LightShaftQuality=1
-+CVars=r.RefractionQuality=1
-
-[iPodTouch5 DeviceProfile]
-DeviceType=IOS
-BaseProfileName=IOS
-+CVars=r.RenderTargetSwitchWorkaround=1
++CVars=sg.PostProcessQuality=3
 
 [iPodTouch6 DeviceProfile]
 DeviceType=IOS
@@ -257,38 +222,31 @@ BaseProfileName=iPhone7
 DeviceType=IOS
 BaseProfileName=IOS
 +CVars=r.MobileContentScaleFactor=2
-+CVars=r.BloomQuality=1
-+CVars=r.DepthOfFieldQuality=1
-+CVars=r.LightShaftQuality=1
-+CVars=r.RefractionQuality=1
++CVars=sg.PostProcessQuality=3
 
 [iPhone6Plus DeviceProfile]
 DeviceType=IOS
 BaseProfileName=IOS
 +CVars=r.MobileContentScaleFactor=2
-+CVars=r.BloomQuality=1
-+CVars=r.DepthOfFieldQuality=1
-+CVars=r.LightShaftQuality=1
-+CVars=r.RefractionQuality=1
++CVars=sg.PostProcessQuality=3
 
 [iPhone6S DeviceProfile]
 DeviceType=IOS
 BaseProfileName=iPhone6
 +CVars=r.MobileContentScaleFactor=2
++CVars=sg.ShadowQuality=2
 
 [iPhone6SPlus DeviceProfile]
 DeviceType=IOS
 BaseProfileName=iPhone6Plus
 +CVars=r.MobileContentScaleFactor=2
++CVars=sg.ShadowQuality=2
 
 [iPhoneSE DeviceProfile]
 DeviceType=IOS
 BaseProfileName=IOS
 +CVars=r.MobileContentScaleFactor=2
-+CVars=r.BloomQuality=1
-+CVars=r.DepthOfFieldQuality=1
-+CVars=r.LightShaftQuality=1
-+CVars=r.RefractionQuality=1
++CVars=sg.PostProcessQuality=3
 
 [iPhone7 DeviceProfile]
 DeviceType=IOS
@@ -303,10 +261,12 @@ BaseProfileName=iPhone6Plus
 [iPhone8 DeviceProfile]
 DeviceType=IOS
 BaseProfileName=iPhone7
++CVars=sg.ShadowQuality=3
 
 [iPhone8Plus DeviceProfile]
 DeviceType=IOS
 BaseProfileName=iPhone7Plus
++CVars=sg.ShadowQuality=3
 
 [iPhoneX DeviceProfile]
 DeviceType=IOS
@@ -346,24 +306,20 @@ BaseProfileName=iPhone8
 [iPadPro DeviceProfile]
 DeviceType=IOS
 BaseProfileName=IOS
-+CVars=r.BloomQuality=1
-+CVars=r.DepthOfFieldQuality=1
-+CVars=r.LightShaftQuality=1
-+CVars=r.RefractionQuality=1
 +CVars=r.MobileContentScaleFactor=1.5
 +CVars=g.TimeoutForBlockOnRenderFence=3000000
 +CVars=r.Decal.StencilSizeThreshold=0.1
 +CVars=r.MetalComputeParameterSize=1024
 +CVars=r.EarlyZPass=3
-+CVars=r.PostProcessAAQuality=4
++CVars=sg.ShadowQuality=3
++CVars=sg.PostProcessQuality=3
++CVars=sg.AntiAliasingQuality=3
 
 [AppleTV DeviceProfile]
 DeviceType=IOS
 BaseProfileName=IOS
-+CVars=r.BloomQuality=1
-+CVars=r.DepthOfFieldQuality=1
-+CVars=r.LightShaftQuality=1
-+CVars=r.RefractionQuality=1
++CVars=sg.PostProcessQuality=3
++CVars=sg.ShadowQuality=3
 
 [AppleTV4K DeviceProfile]
 DeviceType=IOS
@@ -413,243 +369,176 @@ BaseProfileName=IPadPro
 DeviceType=IOS
 BaseProfileName=IPadPro
 
+;----------------
+; Android devices
+
+; These rules are checked in turn and stop at the profile that matches all criteria
 [/Script/AndroidDeviceProfileSelector.AndroidDeviceProfileMatchingRules]
-MatchProfile=(Profile="Android_PowerVR54x",Match=((SourceType=SRC_GpuFamily,CompareType=CMP_Regex,MatchString="PowerVR SGX 54[0-9]")))
-+MatchProfile=(Profile="Android_PowerVR_G6xxx",Match=((SourceType=SRC_GpuFamily,CompareType=CMP_Regex,MatchString="PowerVR Rogue G6[0-9]+")))
-+MatchProfile=(Profile="Android_PowerVRHigh",Match=((SourceType=SRC_GpuFamily,CompareType=CMP_Regex,MatchString="PowerVR Rogue G"),(SourceType=SRC_GlVersion,CompareType=CMP_Regex,MatchString="^OpenGL ES 3\\.")))
-+MatchProfile=(Profile="Android_PowerVRMid",Match=((SourceType=SRC_GpuFamily,CompareType=CMP_Regex,MatchString="PowerVR"),(SourceType=SRC_GlVersion,CompareType=CMP_Regex,MatchString="^OpenGL ES 3\\.")))
-+MatchProfile=(Profile="Android_PowerVRLow",Match=((SourceType=SRC_GpuFamily,CompareType=CMP_Regex,MatchString="PowerVR")))
-+MatchProfile=(Profile="Android_Adreno2xx",Match=((SourceType=SRC_GpuFamily,CompareType=CMP_Regex,MatchString="Adreno \\(TM\\) 2[0-9][0-9]")))
-+MatchProfile=(Profile="Android_Adreno320",Match=((SourceType=SRC_GpuFamily,CompareType=CMP_Equal,MatchString="Adreno (TM) 320")))
-+MatchProfile=(Profile="Android_Adreno330_Ver53",Match=((SourceType=SRC_GpuFamily,CompareType=CMP_Equal,MatchString="Adreno (TM) 330"),(SourceType=SRC_GlVersion,CompareType=CMP_Regex,MatchString="ES 3\\.0 V@([0-9]+)"),(SourceType=SRC_PreviousRegexMatch,CompareType=CMP_GreaterEqual,MatchString="53")))
-+MatchProfile=(Profile="Android_Adreno330",Match=((SourceType=SRC_GpuFamily,CompareType=CMP_Equal,MatchString="Adreno (TM) 330")))
-+MatchProfile=(Profile="Android_Adreno4xx_Nexus6",Match=((SourceType=SRC_DeviceModel,CompareType=CMP_Equal,MatchString="Nexus 6")))
-+MatchProfile=(Profile="Android_Adreno4xx",Match=((SourceType=SRC_GpuFamily,CompareType=CMP_Regex,MatchString="Adreno \\(TM\\) 4[0-9][0-9]")))
-+MatchProfile=(Profile="Android_Adreno51x",Match=((SourceType=SRC_GpuFamily,CompareType=CMP_Regex,MatchString="Adreno \\(TM\\) 51[0-9]")))
-+MatchProfile=(Profile="Android_Adreno5xx_No_Vulkan",Match=((SourceType=SRC_GpuFamily,CompareType=CMP_Regex,MatchString="Adreno \\(TM\\) 5[0-9][0-9]"),(SourceType=SRC_AndroidVersion, CompareType=CMP_Regex,MatchString="([0-9]+).*"),(SourceType=SRC_PreviousRegexMatch,CompareType=CMP_Less,MatchString="8")))
+MatchProfile=(Profile="Android_Adreno4xx",Match=((SourceType=SRC_GpuFamily,CompareType=CMP_Regex,MatchString="Adreno \\(TM\\) 4[0-9][0-9]")))
++MatchProfile=(Profile="Android_Adreno5xx_Low",Match=((SourceType=SRC_GpuFamily,CompareType=CMP_Regex,MatchString="Adreno \\(TM\\) 5[0-1][0-9]")))
 +MatchProfile=(Profile="Android_Adreno5xx",Match=((SourceType=SRC_GpuFamily,CompareType=CMP_Regex,MatchString="Adreno \\(TM\\) 5[0-9][0-9]")))
++MatchProfile=(Profile="Android_Adreno6xx_Vulkan",Match=((SourceType=SRC_GpuFamily,CompareType=CMP_Regex,MatchString="Adreno \\(TM\\) 6[0-9][0-9]")),(SourceType=SRC_AndroidVersion, CompareType=CMP_Regex,MatchString="([0-9]+).*"),(SourceType=SRC_PreviousRegexMatch,CompareType=CMP_GreaterEqual,MatchString="9"),(SourceType=SRC_VulkanAvailable,CompareType=CMP_Equal,MatchString="true")))
 +MatchProfile=(Profile="Android_Adreno6xx",Match=((SourceType=SRC_GpuFamily,CompareType=CMP_Regex,MatchString="Adreno \\(TM\\) 6[0-9][0-9]")))
-+MatchProfile=(Profile="Android_TegraK1",Match=((SourceType=SRC_GpuFamily,CompareType=CMP_Equal,MatchString="NVIDIA Tegra"),(SourceType=SRC_GlVersion,CompareType=CMP_Regex,MatchString="^OpenGL ES 3\\.")))
-+MatchProfile=(Profile="Android_Tegra4",Match=((SourceType=SRC_GpuFamily,CompareType=CMP_Equal,MatchString="NVIDIA Tegra"),(SourceType=SRC_GlVersion,CompareType=CMP_Regex,MatchString="^OpenGL ES 2\\.")))
-+MatchProfile=(Profile="Android_IntelHD_ES3",Match=((SourceType=SRC_GpuFamily,CompareType=CMP_Equal,MatchString="Intel(R) HD Graphics"),(SourceType=SRC_GlVersion,CompareType=CMP_Regex,MatchString="^OpenGL ES 3\\.")))
-+MatchProfile=(Profile="Android_IntelHD",Match=((SourceType=SRC_GpuFamily,CompareType=CMP_Equal,MatchString="Intel(R) HD Graphics"),(SourceType=SRC_GlVersion,CompareType=CMP_Regex,MatchString="^OpenGL ES 2\\.")))
-+MatchProfile=(Profile="Android_Mali_4xx_LenovoS939",Match=((SourceType=SRC_DeviceModel,CompareType=CMP_Equal,MatchString="Lenovo S939")))
-+MatchProfile=(Profile="Android_Mali_4xx_LenovoA7600",Match=((SourceType=SRC_DeviceModel,CompareType=CMP_Regex,MatchString="^Lenovo A7600")))
-+MatchProfile=(Profile="Android_Mali_4xx_WikoHighwaySigns",Match=((SourceType=SRC_DeviceMake,CompareType=CMP_Equal,MatchString="Wiko"),(SourceType=SRC_DeviceModel,CompareType=CMP_Equal,MatchString="Highway Signs")))
-+MatchProfile=(Profile="Android_Mali_4xx_AlcatelPixi3",Match=((SourceType=SRC_DeviceMake,CompareType=CMP_Equal,MatchString="Alcatel"),(SourceType=SRC_DeviceModel,CompareType=CMP_Equal,MatchString="Pixi 3 (4) 3G")))
-+MatchProfile=(Profile="Android_Mali_4xx_KindleFire7",Match=((SourceType=SRC_DeviceModel,CompareType=CMP_Equal,MatchString="KFFOWI")))
-+MatchProfile=(Profile="Android_Mali_4xx",Match=((SourceType=SRC_GpuFamily,CompareType=CMP_Regex,MatchString="^Mali\\-4")))
 +MatchProfile=(Profile="Android_Mali_T6xx",Match=((SourceType=SRC_GpuFamily,CompareType=CMP_Regex,MatchString="^Mali\\-T6")))
 +MatchProfile=(Profile="Android_Mali_T7xx",Match=((SourceType=SRC_GpuFamily,CompareType=CMP_Regex,MatchString="^Mali\\-T7")))
-+MatchProfile=(Profile="Android_Mali_T8xx_No_Vulkan",Match=((SourceType=SRC_GpuFamily,CompareType=CMP_Regex,MatchString="^Mali\\-T8"),(SourceType=SRC_VulkanVersion,CompareType=CMP_Regex,MatchString="^1\\.0\\.([0-9]+)"),(SourceType=SRC_PreviousRegexMatch,CompareType=CMP_Less,MatchString="11")))
-+MatchProfile=(Profile="Android_Mali_T8xx_No_Vulkan",Match=((SourceType=SRC_GpuFamily,CompareType=CMP_Regex,MatchString="^Mali\\-T8"),(SourceType=SRC_AndroidVersion, CompareType=CMP_Regex,MatchString="([0-9]+).*"),(SourceType=SRC_PreviousRegexMatch,CompareType=CMP_Less,MatchString="8")))
 +MatchProfile=(Profile="Android_Mali_T8xx",Match=((SourceType=SRC_GpuFamily,CompareType=CMP_Regex,MatchString="^Mali\\-T8")))
 +MatchProfile=(Profile="Android_Mali_G71",Match=((SourceType=SRC_GpuFamily,CompareType=CMP_Regex,MatchString="^Mali\\-G71")))
++MatchProfile=(Profile="Android_Mali_G72_Vulkan",Match=((SourceType=SRC_GpuFamily,CompareType=CMP_Regex,MatchString="^Mali\\-G72"),(SourceType=SRC_VulkanAvailable,CompareType=CMP_Equal,MatchString="true")))
 +MatchProfile=(Profile="Android_Mali_G72",Match=((SourceType=SRC_GpuFamily,CompareType=CMP_Regex,MatchString="^Mali\\-G72")))
++MatchProfile=(Profile="Android_Mali_G76_Vulkan",Match=((SourceType=SRC_GpuFamily,CompareType=CMP_Regex,MatchString="^Mali\\-G76"),(SourceType=SRC_VulkanAvailable,CompareType=CMP_Equal,MatchString="true")))
 +MatchProfile=(Profile="Android_Mali_G76",Match=((SourceType=SRC_GpuFamily,CompareType=CMP_Regex,MatchString="^Mali\\-G76")))
-+MatchProfile=(Profile="Android_Mid",Match=((SourceType=SRC_GlVersion,CompareType=CMP_Regex,MatchString="^OpenGL ES 3\\.")))
-+MatchProfile=(Profile="Android_Low",Match=((SourceType=SRC_GlVersion,CompareType=CMP_Regex,MatchString="^OpenGL ES 2\\.")))
-
-[/Script/AndroidDeviceProfileSelector.AndroidJavaSurfaceViewDevices]
-SurfaceViewDevices=(Manufacturer="Amazon",Model="Kindle Fire")	; Kindle Fire (2011)
-+SurfaceViewDevices=(Manufacturer="Amazon",Model="KFOT")		; Kindle Fire (2012)
-+SurfaceViewDevices=(Manufacturer="Amazon",Model="KFTT")		; Kindle Fire HD 7" 2012 (WiFi)
-+SurfaceViewDevices=(Manufacturer="Amazon",Model="KFJWA")		; Kindle Fire HD 8.9" 2012 (WAN)
-+SurfaceViewDevices=(Manufacturer="Amazon",Model="KFJWI")		; Kindle Fire HD 8.9" 2012 (WiFi)
-+SurfaceViewDevices=(Manufacturer="HTC",Model="HTC6435LVW")		; HTC One M7
++MatchProfile=(Profile="Android_PowerVR_G6xxx",Match=((SourceType=SRC_GpuFamily,CompareType=CMP_Regex,MatchString="PowerVR Rogue G6[0-9]+")))
++MatchProfile=(Profile="Android_PowerVR_GT7xxx",Match=((SourceType=SRC_GpuFamily,CompareType=CMP_Regex,MatchString="PowerVR Rogue GT7[0-9]+")))
++MatchProfile=(Profile="Android_PowerVR_GE8xxx",Match=((SourceType=SRC_GpuFamily,CompareType=CMP_Regex,MatchString="PowerVR Rogue GE8[0-9]+")))
++MatchProfile=(Profile="Android_PowerVR_GM9xxx_Vulkan",Match=((SourceType=SRC_GpuFamily,CompareType=CMP_Regex,MatchString="PowerVR Rogue GM9[0-9]+"),(SourceType=SRC_VulkanAvailable,CompareType=CMP_Equal,MatchString="true")))
++MatchProfile=(Profile="Android_PowerVR_GM9xxx",Match=((SourceType=SRC_GpuFamily,CompareType=CMP_Regex,MatchString="PowerVR Rogue GM9[0-9]+")))
++MatchProfile=(Profile="Android_TegraK1",Match=((SourceType=SRC_GpuFamily,CompareType=CMP_Equal,MatchString="NVIDIA Tegra"),(SourceType=SRC_GlVersion,CompareType=CMP_Regex,MatchString="^OpenGL ES 3\\.")))
+; Android_Default is the profile of last resort
 
 [Android DeviceProfile]
 DeviceType=Android
 BaseProfileName=
-+CVars=r.MobileContentScaleFactor=1
-+CVars=r.BloomQuality=0
-+CVars=r.DepthOfFieldQuality=0
-+CVars=r.LightShaftQuality=0
-+CVars=r.RefractionQuality=0
-+CVars=r.ShadowQuality=2
++CVars=r.MobileContentScaleFactor=1.0
 +CVars=slate.AbsoluteIndices=1
 +CVars=r.Vulkan.DelayAcquireBackBuffer=2
 +CVars=r.Vulkan.RobustBufferAccess=1
 +CVars=r.Vulkan.DescriptorSetLayoutMode=2
+; Don't enable Vulkan by default. Specific device profiles can set this cvar to 0 to enable Vulkan.
++CVars=r.Android.DisableVulkanSupport=1
 
 [Android_Low DeviceProfile]
 DeviceType=Android
 BaseProfileName=Android
-+CVars=r.MobileContentScaleFactor=0.5
-+CVars=r.ShadowQuality=1
++CVars=r.MobileContentScaleFactor=0.8
+; Scalability groups, see AndroidScalability.ini
++CVars=sg.ViewDistanceQuality=0
++CVars=sg.AntiAliasingQuality=0
++CVars=sg.ShadowQuality=0
++CVars=sg.PostProcessQuality=0
++CVars=sg.TextureQuality=0
 +CVars=sg.EffectsQuality=0
++CVars=sg.FoliageQuality=0
 
 [Android_Mid DeviceProfile]
 DeviceType=Android
 BaseProfileName=Android
-+CVars=r.BloomQuality=1
-+CVars=r.MobileContentScaleFactor=0.8
++CVars=r.MobileContentScaleFactor=1.0
++CVars=sg.ViewDistanceQuality=1
++CVars=sg.AntiAliasingQuality=1
++CVars=sg.ShadowQuality=1
++CVars=sg.PostProcessQuality=1
++CVars=sg.TextureQuality=1
++CVars=sg.EffectsQuality=1
++CVars=sg.FoliageQuality=1
 
 [Android_High DeviceProfile]
 DeviceType=Android
 BaseProfileName=Android
-+CVars=r.BloomQuality=1
-+CVars=r.DepthOfFieldQuality=1
-+CVars=r.LightShaftQuality=1
-+CVars=r.RefractionQuality=1
-+CVars=r.MobileContentScaleFactor=1.0
++CVars=sg.ViewDistanceQuality=2
++CVars=sg.AntiAliasingQuality=2
++CVars=sg.ShadowQuality=2
++CVars=sg.PostProcessQuality=2
++CVars=sg.TextureQuality=2
++CVars=sg.EffectsQuality=2
++CVars=sg.FoliageQuality=2+CVars=r.MobileContentScaleFactor=1.0
 
 [Android_Default DeviceProfile]
 DeviceType=Android
 BaseProfileName=Android_Mid
 
-[Android_Adreno320 DeviceProfile]
-DeviceType=Android
-BaseProfileName=Android_Mid
-
-;This offset needs to be set for the mosaic fallback to work on Galaxy S4 (SAMSUNG-IGH-I337)
-;+CVars=r.DemosaicVposOffset=0.5
+;----------------
+; Adreno devices, Adreno 4xx and above are supported
 
-[Android_Adreno2xx DeviceProfile]
+[Android_Adreno4xx DeviceProfile]
 DeviceType=Android
 BaseProfileName=Android_Low
 
-[Android_Adreno330 DeviceProfile]
+[Android_Adreno5xx_Low DeviceProfile]
 DeviceType=Android
-BaseProfileName=Android_High
-+CVars=OpenGL.UseBufferDiscard=0
-
-[Android_Adreno330_Ver53 DeviceProfile]
-DeviceType=Android
-BaseProfileName=Android_Adreno330
-+CVars=r.DisjointTimerQueries=0
-
-[Android_Adreno51x DeviceProfile]
-DeviceType=Android
-BaseProfileName=Android_High
+BaseProfileName=Android_Low
 
 [Android_Adreno5xx DeviceProfile]
 DeviceType=Android
-BaseProfileName=Android_High
-+CVars=r.DisjointTimerQueries=1
-
-[Android_Adreno5xx_No_Vulkan DeviceProfile]
-DeviceType=Android
-BaseProfileName=Android_Adreno5xx
-; There are several issues (vulkan sub-passes, occlusion queries) on devices running Android 7 and earier
-+CVars=r.Android.DisableVulkanSupport=1
+BaseProfileName=Android_Mid
 
 [Android_Adreno6xx DeviceProfile]
 DeviceType=Android
 BaseProfileName=Android_High
-+CVars=r.DisjointTimerQueries=1
-
-[Android_Adreno4xx DeviceProfile]
-DeviceType=Android
-BaseProfileName=Android_High
 
-[Android_Adreno4xx_Nexus6 DeviceProfile]
+[Android_Adreno6xx_Vulkan DeviceProfile]
 DeviceType=Android
-BaseProfileName=Android_Adreno4xx
+BaseProfileName=Android_Adreno6xx
+; enable Vulkan
++CVars=r.Android.DisableVulkanSupport=0
 
-[Android_PowerVR54x DeviceProfile]
-DeviceType=Android
-BaseProfileName=Android_Mid
-+CVars=r.Mobile.ForceDepthResolve=1
-+CVars=r.Android.DisableASTCSupport=1
-
-[Android_PowerVRMid DeviceProfile]
-DeviceType=Android
-BaseProfileName=Android_Mid
-+CVars=r.Mobile.ForceDepthResolve=1
-+CVars=r.Android.DisableASTCSupport=1
+;----------------
+; PowerVR devices
+; some PowerVR G6xxx devices are known to have problems with ES 3.1 on UE4
 
 [Android_PowerVR_G6xxx DeviceProfile]
 DeviceType=Android
-BaseProfileName=Android_PowerVRMid
-; Disable ES3.1 support on PoverVR G6xxx devices
-+CVars=r.Android.DisableOpenGLES31Support=1
-
-[Android_PowerVRLow DeviceProfile]
-DeviceType=Android
-BaseProfileName=Android_Low
-+CVars=r.Mobile.ForceDepthResolve=1
-+CVars=r.Android.DisableASTCSupport=1
-
-[Android_PowerVRHigh DeviceProfile]
-DeviceType=Android
-BaseProfileName=Android_High
-
-[Android_IntelHD DeviceProfile]
-DeviceType=Android
 BaseProfileName=Android_Low
 
-[Android_IntelHD_ES3 DeviceProfile]
+[Android_PowerVR_GT7xxx DeviceProfile]
 DeviceType=Android
-BaseProfileName=Android
-+CVars=r.BloomQuality=1
-
-[Android_Mali_4xx DeviceProfile]
-DeviceType=Android
-BaseProfileName=Android_Low
-
-[Android_Mali_4xx_LenovoS939 DeviceProfile]
-DeviceType=Android
-BaseProfileName=Android_Mali_4xx
-+CVars=r.AndroidDisableThreadedRendering=1
+BaseProfileName=Android_Mid
 
-[Android_Mali_4xx_LenovoA7600 DeviceProfile]
+[Android_PowerVR_GE8xxx DeviceProfile]
 DeviceType=Android
-BaseProfileName=Android_Mali_4xx
-+CVars=r.AndroidDisableThreadedRendering=1
+BaseProfileName=Android_Mid
 
-[Android_Mali_4xx_WikoHighwaySigns DeviceProfile]
+[Android_PowerVR_GM9xxx DeviceProfile]
 DeviceType=Android
-BaseProfileName=Android_Mali_4xx
-+CVars=r.AndroidDisableThreadedRendering=1
+BaseProfileName=Android_High
 
-[Android_Mali_4xx_AlcatelPixi3 DeviceProfile]
+[Android_PowerVR_GM9xxx_Vulkan DeviceProfile]
 DeviceType=Android
-BaseProfileName=Android_Mali_4xx
-+CVars=r.AndroidDisableThreadedRendering=1
+BaseProfileName=Android_PowerVR_GM9xxx
+; enable Vulkan
++CVars=r.Android.DisableVulkanSupport=0
 
-[Android_Mali_4xx_KindleFire7 DeviceProfile]
-DeviceType=Android
-BaseProfileName=Android_Mali_4xx
-+CVars=r.AndroidDisableThreadedRendering=1
-+CVars=r.AndroidDisableThreadedRenderingFirstLoad=1
+;----------------
+; Mali devices, we support from Mali-T6xx and above
 
 [Android_Mali_T6xx DeviceProfile]
 DeviceType=Android
-BaseProfileName=Android_Mid
+BaseProfileName=Android_Low
 
 [Android_Mali_T7xx DeviceProfile]
 DeviceType=Android
-BaseProfileName=Android_High
+BaseProfileName=Android_Low
 
 [Android_Mali_T8xx DeviceProfile]
 DeviceType=Android
-BaseProfileName=Android_High
+BaseProfileName=Android_Mid
 
 [Android_Mali_G71 DeviceProfile]
 DeviceType=Android
-BaseProfileName=Android_High
+BaseProfileName=Android_Mid
 
 [Android_Mali_G72 DeviceProfile]
 DeviceType=Android
 BaseProfileName=Android_High
 
-[Android_Mali_G76 DeviceProfile]
+[Android_Mali_G72_Vulkan DeviceProfile]
 DeviceType=Android
-BaseProfileName=Android_High
+BaseProfileName=Android_Mali_G72
+; enable Vulkan
++CVars=r.Android.DisableVulkanSupport=0
 
-[Android_TegraK1 DeviceProfile]
+[Android_Mali_G76 DeviceProfile]
 DeviceType=Android
 BaseProfileName=Android_High
 
-[Android_Tegra4 DeviceProfile]
+[Android_Mali_G76_Vulkan DeviceProfile]
 DeviceType=Android
-BaseProfileName=Android_Mid
+BaseProfileName=Android_Mali_G76
+; enable Vulkan
++CVars=r.Android.DisableVulkanSupport=0
 
-[Android_Mali_T8xx_No_Vulkan DeviceProfile]
+[Android_TegraK1 DeviceProfile]
 DeviceType=Android
-BaseProfileName=Android_Mali_T8xx
-+CVars=r.Android.DisableVulkanSupport=1
+BaseProfileName=Android_High
 
 [Lumin DeviceProfile]
 DeviceType=Lumin
diff --git a/Engine/Config/IOS/IOSScalability.ini b/Engine/Config/IOS/IOSScalability.ini
--- a/Engine/Config/IOS/IOSScalability.ini
+++ b/Engine/Config/IOS/IOSScalability.ini
@@ -1,3 +1,110 @@
+; ------------------------------------------------------------------------------
+; Foliage settings
+; ------------------------------------------------------------------------------
+
+[FoliageQuality@0]
+foliage.DensityScale=0.5
+grass.DensityScale=0.5
+
+[FoliageQuality@1]
+grass.densityScale=0.8
+grass.CullDistanceScale=0.8
+foliage.DensityScale=0.8
+
+[FoliageQuality@2]
+grass.densityScale=1.0
+grass.CullDistanceScale=1.0
+foliage.DensityScale=1.0
+
+[FoliageQuality@3]
+grass.densityScale=1.0
+grass.CullDistanceScale=1.0
+foliage.DensityScale=1.0
+
+; ------------------------------------------------------------------------------
+; View distance settings
+; ------------------------------------------------------------------------------
+
+[ViewDistanceQuality@0]
+r.ViewDistanceScale=0.70
+
+[ViewDistanceQuality@1]
+r.ViewDistanceScale=1.0
+
+[ViewDistanceQuality@2]
+r.ViewDistanceScale=1.0
+
+[ViewDistanceQuality@3]
+r.ViewDistanceScale=1.0
+
+; ------------------------------------------------------------------------------
+; Shadow settings
+; ------------------------------------------------------------------------------
+
+[ShadowQuality@0]
+r.ShadowQuality=0
+r.Shadow.MaxCSMResolution=1024
+r.Shadow.DistanceScale=0.5
+r.Shadow.CSM.MaxCascades=1
+r.Shadow.RadiusThreshold=0.06
+
+[ShadowQuality@1]
+r.ShadowQuality=1
+r.Shadow.MaxCSMResolution=1024
+r.Shadow.DistanceScale=0.7
+r.Shadow.CSM.MaxCascades=1
+r.Shadow.RadiusThreshold=0.06
+
+[ShadowQuality@2]
+r.ShadowQuality=1
+r.Shadow.MaxCSMResolution=2048
+r.Shadow.DistanceScale=1.0
+r.Shadow.CSM.MaxCascades=2
+r.Shadow.RadiusThreshold=0.06
+
+[ShadowQuality@3]
+r.ShadowQuality=2
+r.Shadow.MaxCSMResolution=2048
+r.Shadow.DistanceScale=1.0
+r.Shadow.CSM.MaxCascades=2
+r.Shadow.RadiusThreshold=0.06
+
+; ------------------------------------------------------------------------------
+; Postprocess settings
+; ------------------------------------------------------------------------------
+
+[PostProcessQuality@0]
+r.Tonemapper.Quality=1
+r.BloomQuality=0
+r.DepthOfFieldQuality=0
+r.LightShaftQuality=0
+r.RefractionQuality=0
+Slate.ForceBackgroundBlurLowQualityOverride=1
+
+[PostProcessQuality@1]
+r.Tonemapper.Quality=2
+r.BloomQuality=1
+r.DepthOfFieldQuality=0
+r.LightShaftQuality=0
+r.RefractionQuality=0
+Slate.ForceBackgroundBlurLowQualityOverride=1
+
+[PostProcessQuality@2]
+r.Tonemapper.Quality=3
+r.BloomQuality=1
+r.DepthOfFieldQuality=0
+r.LightShaftQuality=0
+r.RefractionQuality=1
+Slate.ForceBackgroundBlurLowQualityOverride=1
+
+[PostProcessQuality@3]
+r.Tonemapper.Quality=3
+r.BloomQuality=1
+r.DepthOfFieldQuality=1
+r.LightShaftQuality=1
+r.RefractionQuality=1
+Slate.ForceBackgroundBlurLowQualityOverride=1
+
 ; ------------------------------------------------------------------------------
 ; Effects settings
 ; ------------------------------------------------------------------------------
@@ -45,3 +152,27 @@ r.SkyAtmosphere.FastSkyLUT.SampleCountMax=8.0
 r.SkyAtmosphere.FastSkyLUT.Width=96
 r.SkyAtmosphere.FastSkyLUT.Height=50
 r.SkyAtmosphere.TransmittanceLUT.UseSmallFormat=1
+
+; ------------------------------------------------------------------------------
+; Texture settings
+; ------------------------------------------------------------------------------
+
+[TextureQuality@0]
+r.MaxAnisotropy=1
+
+[TextureQuality@1]
+r.MaxAnisotropy=1
+
+[TextureQuality@2]
+r.MaxAnisotropy=2
+
+[TextureQuality@3]
+r.MaxAnisotropy=4
+
+[ScalabilityPreview]
+; don't want to preview any texture streaming changes in the editor
+BlacklistCVars=r.MaxAnisotropy
++BlacklistCVars=r.Streaming.Boost
++BlacklistCVars=r.Streaming.AmortizeCPUToGPUCopy
++BlacklistCVars=r.Streaming.MaxNumTexturesToStreamPerFrame
++BlacklistCVars=r.Streaming.PoolSize
\ No newline at end of file
