Make using shader code library the default setting.

diff --git a/Engine/Config/BaseGame.ini b/Engine/Config/BaseGame.ini
--- a/Engine/Config/BaseGame.ini
+++ b/Engine/Config/BaseGame.ini
@@ -73,7 +73,7 @@ InternationalizationPreset=English
 DefaultCulture=en
 bSkipEditorContent=false
 bSharedMaterialNativeLibraries=False
-bShareMaterialShaderCode=False
+bShareMaterialShaderCode=True
 bSkipMovies=False
 bPakUsesSecondaryOrder=True
 
