Set SensorLandscape orientation as default for Android apps
Set Fullscreen as default for Android Apps

diff --git a/Engine/Config/BaseEngine.ini b/Engine/Config/BaseEngine.ini
--- a/Engine/Config/BaseEngine.ini
+++ b/Engine/Config/BaseEngine.ini
@@ -2055,7 +2055,7 @@ bPackageDataInsideApk=false
 bUseExternalFilesDir=false
 bPublicLogFiles=true
 bCreateAllPlatformsInstall=false
-Orientation=Landscape
+Orientation=SensorLandscape
 InstallLocation=InternalOnly
 DepthBufferPreference=Default
 PackageName=com.YourCompany.[PROJECT]
@@ -2098,6 +2098,7 @@ bEnableUniversalAPK=False
 bBundleABISplit=True
 bBundleLanguageSplit=True
 bBundleDensitySplit=True
+bFullScreen=True
 
 [/Script/AndroidPlatformEditor.AndroidSDKSettings]
 SDKAPILevel=latest
