Cine Camera: Update 16:9 Digital Film SensorWidth and SensorHeight to 23.76x13.365 to match Arri Alexa from https://vfxcamdb.com/arri-alexa/

diff --git a/Engine/Config/BaseEngine.ini b/Engine/Config/BaseEngine.ini
--- a/Engine/Config/BaseEngine.ini
+++ b/Engine/Config/BaseEngine.ini
@@ -2304,7 +2304,7 @@ ServerUpdateCameraTimeout=2.0
 
 [/Script/CinematicCamera.CineCameraComponent]
 +FilmbackPresets=(Name="16:9 Film",FilmbackSettings=(SensorWidth=24.892,SensorHeight=14.002))
-+FilmbackPresets=(Name="16:9 Digital Film",FilmbackSettings=(SensorWidth=23.66,SensorHeight=13.3))
++FilmbackPresets=(Name="16:9 Digital Film",FilmbackSettings=(SensorWidth=23.76,SensorHeight=13.365))
 +FilmbackPresets=(Name="16:9 DSLR",FilmbackSettings=(SensorWidth=36,SensorHeight=20.25))
 +FilmbackPresets=(Name="Super 8mm",FilmbackSettings=(SensorWidth=5.79,SensorHeight=4.01))
 +FilmbackPresets=(Name="Super 16mm",FilmbackSettings=(SensorWidth=12.52,SensorHeight=7.58))
