Cine Camera: Update 16:9 Film SensorWidth and SensorHeight to 24.00 x 13.5

diff --git a/Engine/Config/BaseEngine.ini b/Engine/Config/BaseEngine.ini
--- a/Engine/Config/BaseEngine.ini
+++ b/Engine/Config/BaseEngine.ini
@@ -2303,7 +2303,7 @@ GameHUDBind=(Key=Tilde,bShift=False,bCtrl=True,bAlt=False,bCmd=False)
 ServerUpdateCameraTimeout=2.0
 
 [/Script/CinematicCamera.CineCameraComponent]
-+FilmbackPresets=(Name="16:9 Film",FilmbackSettings=(SensorWidth=24.892,SensorHeight=14.002))
++FilmbackPresets=(Name="16:9 Film",FilmbackSettings=(SensorWidth=24.00,SensorHeight=13.5))
 +FilmbackPresets=(Name="16:9 Digital Film",FilmbackSettings=(SensorWidth=23.76,SensorHeight=13.365))
 +FilmbackPresets=(Name="16:9 DSLR",FilmbackSettings=(SensorWidth=36,SensorHeight=20.25))
 +FilmbackPresets=(Name="Super 8mm",FilmbackSettings=(SensorWidth=5.79,SensorHeight=4.01))
