Niagara: Disable sequencer clean playback mode by default

diff --git a/Engine/Config/BaseEditorPerProjectUserSettings.ini b/Engine/Config/BaseEditorPerProjectUserSettings.ini
--- a/Engine/Config/BaseEditorPerProjectUserSettings.ini
+++ b/Engine/Config/BaseEditorPerProjectUserSettings.ini
@@ -890,7 +890,7 @@ bKeepPlayRangeInSectionBounds=false
 bKeepCursorInPlayRange=false
 bShowRangeSlider=true
 LoopMode=SLM_Loop
-
+bCleanPlaybackMode=false
 
 [/Script/LevelSequenceEditor.LevelSequenceEditorSettings]
 +TrackSettings=(MatchingActorClass=/Script/Engine.StaticMeshActor,DefaultTracks=(/Script/MovieSceneTracks.MovieScene3DTransformTrack))
