Use native shader libraries by default

diff --git a/Engine/Config/BaseGame.ini b/Engine/Config/BaseGame.ini
--- a/Engine/Config/BaseGame.ini
+++ b/Engine/Config/BaseGame.ini
@@ -72,7 +72,7 @@ InternationalizationPreset=English
 +CulturesToStage=en
 DefaultCulture=en
 bSkipEditorContent=false
-bSharedMaterialNativeLibraries=False
+bSharedMaterialNativeLibraries=True
 bShareMaterialShaderCode=True
 bSkipMovies=False
 bPakUsesSecondaryOrder=True
