Disabled Niagra Compute on devices with A10x GPU's and below due to lack of consistant RW texture support

diff --git a/Engine/Config/BaseDeviceProfiles.ini b/Engine/Config/BaseDeviceProfiles.ini
--- a/Engine/Config/BaseDeviceProfiles.ini
+++ b/Engine/Config/BaseDeviceProfiles.ini
@@ -166,97 +166,115 @@ BaseProfileName=
 DeviceType=IOS
 BaseProfileName=IOS
 +CVars=r.RenderTargetSwitchWorkaround=1
++CVars=fx.NiagaraAllowComputeShaders=0
 
 [iPad3 DeviceProfile]
 DeviceType=IOS
 BaseProfileName=IOS
 +CVars=r.RenderTargetSwitchWorkaround=1
++CVars=fx.NiagaraAllowComputeShaders=0
 
 [iPad4 DeviceProfile]
 DeviceType=IOS
 BaseProfileName=IOS
 +CVars=r.RenderTargetSwitchWorkaround=1
++CVars=fx.NiagaraAllowComputeShaders=0
 
 [iPadAir DeviceProfile]
 DeviceType=IOS
 BaseProfileName=IOS
 +CVars=sg.PostProcessQuality=1
++CVars=fx.NiagaraAllowComputeShaders=0
 
 [iPadAir2 DeviceProfile]
 DeviceType=IOS
 BaseProfileName=IOS
 +CVars=sg.PostProcessQuality=3
 +CVars=sg.ShadowQuality=2
++CVars=fx.NiagaraAllowComputeShaders=0
 
 [iPadMini DeviceProfile]
 DeviceType=IOS
 BaseProfileName=IOS
 +CVars=r.RenderTargetSwitchWorkaround=1
++CVars=fx.NiagaraAllowComputeShaders=0
 
 [iPadMini2 DeviceProfile]
 DeviceType=IOS
 BaseProfileName=iPadAir
++CVars=fx.NiagaraAllowComputeShaders=0
 
 [iPadMini4 DeviceProfile]
 DeviceType=IOS
 BaseProfileName=iPadMini2
 +CVars=sg.PostProcessQuality=2
++CVars=fx.NiagaraAllowComputeShaders=0
 
 [iPhone5S DeviceProfile]
 DeviceType=IOS
 BaseProfileName=IOS
 +CVars=r.MobileContentScaleFactor=2
 +CVars=sg.PostProcessQuality=3
++CVars=fx.NiagaraAllowComputeShaders=0
 
 [iPodTouch6 DeviceProfile]
 DeviceType=IOS
 BaseProfileName=iPhone6
 +CVars=r.MobileContentScaleFactor=2
++CVars=fx.NiagaraAllowComputeShaders=0
 
 [iPodTouch7 DeviceProfile]
 DeviceType=IOS
 BaseProfileName=iPhone7
 +CVars=r.MobileContentScaleFactor=2
++CVars=fx.NiagaraAllowComputeShaders=0
 
 [iPhone6 DeviceProfile]
 DeviceType=IOS
 BaseProfileName=IOS
 +CVars=r.MobileContentScaleFactor=2
 +CVars=sg.PostProcessQuality=3
++CVars=fx.NiagaraAllowComputeShaders=0
 
 [iPhone6Plus DeviceProfile]
 DeviceType=IOS
 BaseProfileName=IOS
 +CVars=r.MobileContentScaleFactor=2
 +CVars=sg.PostProcessQuality=3
++CVars=fx.NiagaraAllowComputeShaders=0
 
 [iPhone6S DeviceProfile]
 DeviceType=IOS
 BaseProfileName=iPhone6
 +CVars=r.MobileContentScaleFactor=2
 +CVars=sg.ShadowQuality=2
++CVars=fx.NiagaraAllowComputeShaders=0
 
 [iPhone6SPlus DeviceProfile]
 DeviceType=IOS
 BaseProfileName=iPhone6Plus
 +CVars=r.MobileContentScaleFactor=2
 +CVars=sg.ShadowQuality=2
++CVars=fx.NiagaraAllowComputeShaders=0
 
 [iPhoneSE DeviceProfile]
 DeviceType=IOS
 BaseProfileName=IOS
 +CVars=r.MobileContentScaleFactor=2
 +CVars=sg.PostProcessQuality=3
++CVars=fx.NiagaraAllowComputeShaders=0
 
 [iPhone7 DeviceProfile]
 DeviceType=IOS
 BaseProfileName=iPhone6
 +CVars=r.MobileContentScaleFactor=2
++CVars=fx.NiagaraAllowComputeShaders=0
 
 [iPhone7Plus DeviceProfile]
 DeviceType=IOS
 BaseProfileName=iPhone6Plus
 +CVars=r.MobileContentScaleFactor=2
++CVars=fx.NiagaraAllowComputeShaders=0
 
 [iPhone8 DeviceProfile]
 DeviceType=IOS
@@ -328,6 +346,7 @@ BaseProfileName=AppleTV
 [iPadPro97 DeviceProfile]
 DeviceType=IOS
 BaseProfileName=IPadPro
++CVars=fx.NiagaraAllowComputeShaders=0
 
 [iPadPro129 DeviceProfile]
 DeviceType=IOS
