Added disable flag for SM5 so that specific devices can be targetted for desktop quality rendering

diff --git a/Engine/Config/BaseDeviceProfiles.ini b/Engine/Config/BaseDeviceProfiles.ini
--- a/Engine/Config/BaseDeviceProfiles.ini
+++ b/Engine/Config/BaseDeviceProfiles.ini
@@ -394,6 +394,11 @@ BaseProfileName=IPadPro
 ; These rules are checked in turn and stop at the profile that matches all criteria
 [/Script/AndroidDeviceProfileSelector.AndroidDeviceProfileMatchingRules]
 MatchProfile=(Profile="Android_Adreno4xx",Match=((SourceType=SRC_GpuFamily,CompareType=CMP_Regex,MatchString="Adreno \\(TM\\) 4[0-9][0-9]")))
+
+; Example profiles for Vulkan SM5
+;+MatchProfile=(Profile="Android_Vulkan_SM5",Match=((SourceType=SRC_GpuFamily,CompareType=CMP_Regex,MatchString="Adreno \\(TM\\) 6[0-9][0-9]")),(SourceType=SRC_AndroidVersion, CompareType=CMP_Regex,MatchString="([0-9]+).*"),(SourceType=SRC_PreviousRegexMatch,CompareType=CMP_GreaterEqual,MatchString="9"),(SourceType=SRC_VulkanAvailable,CompareType=CMP_Equal,MatchString="true")))
+;+MatchProfile=(Profile="Android_Vulkan_SM5",Match=((SourceType=SRC_GpuFamily,CompareType=CMP_Regex,MatchString="^Mali\\-G76"),(SourceType=SRC_VulkanAvailable,CompareType=CMP_Equal,MatchString="true")))
+
 +MatchProfile=(Profile="Android_Adreno5xx_Low",Match=((SourceType=SRC_GpuFamily,CompareType=CMP_Regex,MatchString="Adreno \\(TM\\) 5[0-1][0-9]")))
 +MatchProfile=(Profile="Android_Adreno5xx",Match=((SourceType=SRC_GpuFamily,CompareType=CMP_Regex,MatchString="Adreno \\(TM\\) 5[0-9][0-9]")))
 +MatchProfile=(Profile="Android_Adreno6xx_Vulkan",Match=((SourceType=SRC_GpuFamily,CompareType=CMP_Regex,MatchString="Adreno \\(TM\\) 6[0-9][0-9]")),(SourceType=SRC_AndroidVersion, CompareType=CMP_Regex,MatchString="([0-9]+).*"),(SourceType=SRC_PreviousRegexMatch,CompareType=CMP_GreaterEqual,MatchString="9"),(SourceType=SRC_VulkanAvailable,CompareType=CMP_Equal,MatchString="true")))
@@ -424,6 +429,7 @@ BaseProfileName=
 +CVars=r.Vulkan.DescriptorSetLayoutMode=2
 ; Don't enable Vulkan by default. Specific device profiles can set this cvar to 0 to enable Vulkan.
 +CVars=r.Android.DisableVulkanSupport=1
++CVars=r.Android.DisableVulkanSM5Support=1
 
 [Android_Low DeviceProfile]
 DeviceType=Android
@@ -459,7 +465,40 @@ BaseProfileName=Android
 +CVars=sg.PostProcessQuality=2
 +CVars=sg.TextureQuality=2
 +CVars=sg.EffectsQuality=2
-+CVars=sg.FoliageQuality=2+CVars=r.MobileContentScaleFactor=1.0
++CVars=sg.FoliageQuality=2
++CVars=r.MobileContentScaleFactor=1.0
+
+; Example base settings for Vulkan SM5 on Android
+[Android_Vulkan_SM5 DeviceProfile]
+DeviceType=Android
+BaseProfileName=Android
++CVars=sg.ViewDistanceQuality=2
++CVars=sg.AntiAliasingQuality=1
++CVars=sg.ShadowQuality=2
++CVars=sg.PostProcessQuality=2
++CVars=sg.TextureQuality=2
++CVars=sg.EffectsQuality=2
++CVars=sg.FoliageQuality=2
++CVars=sg.ShadingQuality=2
++CVars=r.BloomQuality=2
++CVars=r.LightShaftQuality=1
+; Shadows
++CVars=r.Shadow.MaxResolution=2048
++CVars=r.Shadow.MaxCSMResolution=2048
++CVars=r.Shadow.WholeSceneShadowCacheMb=40
++CVars=r.Shadow.CachedShadowsCastFromMovablePrimitives=0
++CVars=r.Shadow.MaxNumPointShadowCacheUpdatesPerFrame=1
++CVars=r.Shadow.MaxNumSpotShadowCacheUpdatesPerFrame=1
++CVars=r.Shadow.DistanceScale=1.0
++CVars=r.Shadow.CSM.MaxCascades=2 
++CVars=r.ShadowQuality=2
++CVars=r.Shadow.CSMShadowDistanceFadeoutMultiplier=2.5
++CVars=r.SSS.Quality=0
++CVars=r.SSS.Scale=0
++CVars=r.SSR.Quality=0
++CVars=r.VirtualTextures=1
++CVars=r.Android.DisableVulkanSM5Support=0
++CVars=r.Android.DisableVulkanSupport=0
 
 [Android_Default DeviceProfile]
 DeviceType=Android
diff --git a/Engine/Source/Runtime/Core/Private/Android/AndroidPlatformMisc.cpp b/Engine/Source/Runtime/Core/Private/Android/AndroidPlatformMisc.cpp
--- a/Engine/Source/Runtime/Core/Private/Android/AndroidPlatformMisc.cpp
+++ b/Engine/Source/Runtime/Core/Private/Android/AndroidPlatformMisc.cpp
@@ -2179,11 +2179,36 @@ bool FAndroidMisc::ShouldUseVulkan()
 
 bool FAndroidMisc::ShouldUseDesktopVulkan()
 {
-	// @todo Lumin: Double check all this stuff after merging general android Vulkan SM5 from main
-	bool bSupportsVulkanSM5 = false;
-	GConfig->GetBool(TEXT("/Script/AndroidRuntimeSettings.AndroidRuntimeSettings"), TEXT("bSupportsVulkanSM5"), bSupportsVulkanSM5, GEngineIni);
+	static int CachedShouldUseDesktopVulkan = -1;
 
-	return bSupportsVulkanSM5;
+	if (CachedShouldUseDesktopVulkan == -1)
+	{
+		CachedShouldUseDesktopVulkan = 0;
+
+		bool bSupportsVulkanSM5 = false;
+		GConfig->GetBool(TEXT("/Script/AndroidRuntimeSettings.AndroidRuntimeSettings"), TEXT("bSupportsVulkanSM5"), bSupportsVulkanSM5, GEngineIni);
+
+		const bool bVulkanSM5Enabled = bSupportsVulkanSM5;
+
+		static const auto CVarDisableVulkanSM5 = IConsoleManager::Get().FindTConsoleVariableDataInt(TEXT("r.Android.DisableVulkanSM5Support"));
+		const bool bVulkanSM5Disabled = CVarDisableVulkanSM5->GetValueOnAnyThread() == 1;
+
+		if (bVulkanSM5Enabled && !bVulkanSM5Disabled)
+		{
+			CachedShouldUseDesktopVulkan = 1;
+			UE_LOG(LogAndroid, Log, TEXT("Vulkan SM5 RHI will be used!"));
+		}
+		else if(bVulkanSM5Disabled)
+		{
+			UE_LOG(LogAndroid, Log, TEXT("Vulkan SM5 is available but disabled for this device."));
+		}
+		else if (!bVulkanSM5Enabled)
+		{
+			UE_LOG(LogAndroid, Log, TEXT("** Vulkan SM5 support is not available (Driver, RHI or shaders are missing, or disabled by cmdline, see above logging for details)"));
+		}
+	}
+
+	return CachedShouldUseDesktopVulkan;
 }
 
 FString FAndroidMisc::GetVulkanVersion()
diff --git a/Engine/Source/Runtime/Core/Private/HAL/ConsoleManager.cpp b/Engine/Source/Runtime/Core/Private/HAL/ConsoleManager.cpp
--- a/Engine/Source/Runtime/Core/Private/HAL/ConsoleManager.cpp
+++ b/Engine/Source/Runtime/Core/Private/HAL/ConsoleManager.cpp
@@ -2711,6 +2711,14 @@ static TAutoConsoleVariable<int32> CVarDisableVulkanSupport(
 	TEXT("  1 = vulkan will be disabled, opengl fall back will be used."),
 	ECVF_ReadOnly);
 
+static TAutoConsoleVariable<int32> CVarDisableVulkanSM5Support(
+	TEXT("r.Android.DisableVulkanSM5Support"),
+	0,
+	TEXT("Disable support for vulkan API. (Android Only)\n")
+	TEXT("  0 = Vulkan SM5 API will be used (providing device and project supports it) [default]\n")
+	TEXT("  1 = Vulkan SM5 will be disabled, Vulkan or OpenGL fall back will be used."),
+	ECVF_ReadOnly);
+
 static TAutoConsoleVariable<int32> CVarDisableOpenGLES31Support(
 	TEXT("r.Android.DisableOpenGLES31Support"),
 	0,
diff --git a/Engine/Source/Runtime/RHI/Private/Android/AndroidDynamicRHI.cpp b/Engine/Source/Runtime/RHI/Private/Android/AndroidDynamicRHI.cpp
--- a/Engine/Source/Runtime/RHI/Private/Android/AndroidDynamicRHI.cpp
+++ b/Engine/Source/Runtime/RHI/Private/Android/AndroidDynamicRHI.cpp
@@ -12,7 +12,7 @@ FDynamicRHI* PlatformCreateDynamicRHI()
 	IDynamicRHIModule* DynamicRHIModule = NULL;
 	ERHIFeatureLevel::Type RequestedFeatureLevel = ERHIFeatureLevel::Num;
 
-	if (FPlatformMisc::ShouldUseVulkan())
+	if (FPlatformMisc::ShouldUseVulkan() || FPlatformMisc::ShouldUseDesktopVulkan())
 	{
 		// Vulkan is required, release the EGL created by FAndroidAppEntry::PlatformInit.
 		FAndroidAppEntry::ReleaseEGL();
