Fix for crash on older iOS devices when using histogram CS

diff --git a/Engine/Config/BaseDeviceProfiles.ini b/Engine/Config/BaseDeviceProfiles.ini
--- a/Engine/Config/BaseDeviceProfiles.ini
+++ b/Engine/Config/BaseDeviceProfiles.ini
@@ -167,24 +167,28 @@ DeviceType=IOS
 BaseProfileName=IOS
 +CVars=r.RenderTargetSwitchWorkaround=1
 +CVars=fx.NiagaraAllowComputeShaders=0
++CVars=r.Mobile.EyeAdaptation=0
 
 [iPad3 DeviceProfile]
 DeviceType=IOS
 BaseProfileName=IOS
 +CVars=r.RenderTargetSwitchWorkaround=1
 +CVars=fx.NiagaraAllowComputeShaders=0
++CVars=r.Mobile.EyeAdaptation=0
 
 [iPad4 DeviceProfile]
 DeviceType=IOS
 BaseProfileName=IOS
 +CVars=r.RenderTargetSwitchWorkaround=1
 +CVars=fx.NiagaraAllowComputeShaders=0
++CVars=r.Mobile.EyeAdaptation=0
 
 [iPadAir DeviceProfile]
 DeviceType=IOS
 BaseProfileName=IOS
 +CVars=sg.PostProcessQuality=1
 +CVars=fx.NiagaraAllowComputeShaders=0
++CVars=r.Mobile.EyeAdaptation=0
 
 [iPadAir2 DeviceProfile]
 DeviceType=IOS
@@ -192,23 +196,27 @@ BaseProfileName=IOS
 +CVars=sg.PostProcessQuality=3
 +CVars=sg.ShadowQuality=2
 +CVars=fx.NiagaraAllowComputeShaders=0
++CVars=r.Mobile.EyeAdaptation=0
 
 [iPadMini DeviceProfile]
 DeviceType=IOS
 BaseProfileName=IOS
 +CVars=r.RenderTargetSwitchWorkaround=1
 +CVars=fx.NiagaraAllowComputeShaders=0
++CVars=r.Mobile.EyeAdaptation=0
 
 [iPadMini2 DeviceProfile]
 DeviceType=IOS
 BaseProfileName=iPadAir
 +CVars=fx.NiagaraAllowComputeShaders=0
++CVars=r.Mobile.EyeAdaptation=0
 
 [iPadMini4 DeviceProfile]
 DeviceType=IOS
 BaseProfileName=iPadMini2
 +CVars=sg.PostProcessQuality=2
 +CVars=fx.NiagaraAllowComputeShaders=0
++CVars=r.Mobile.EyeAdaptation=0
 
 [iPhone5S DeviceProfile]
 DeviceType=IOS
@@ -216,18 +224,21 @@ BaseProfileName=IOS
 +CVars=r.MobileContentScaleFactor=2
 +CVars=sg.PostProcessQuality=3
 +CVars=fx.NiagaraAllowComputeShaders=0
++CVars=r.Mobile.EyeAdaptation=0
 
 [iPodTouch6 DeviceProfile]
 DeviceType=IOS
 BaseProfileName=iPhone6
 +CVars=r.MobileContentScaleFactor=2
 +CVars=fx.NiagaraAllowComputeShaders=0
++CVars=r.Mobile.EyeAdaptation=0
 
 [iPodTouch7 DeviceProfile]
 DeviceType=IOS
 BaseProfileName=iPhone7
 +CVars=r.MobileContentScaleFactor=2
 +CVars=fx.NiagaraAllowComputeShaders=0
++CVars=r.Mobile.EyeAdaptation=0
 
 [iPhone6 DeviceProfile]
 DeviceType=IOS
@@ -235,6 +246,7 @@ BaseProfileName=IOS
 +CVars=r.MobileContentScaleFactor=2
 +CVars=sg.PostProcessQuality=3
 +CVars=fx.NiagaraAllowComputeShaders=0
++CVars=r.Mobile.EyeAdaptation=0
 
 [iPhone6Plus DeviceProfile]
 DeviceType=IOS
@@ -338,6 +350,7 @@ DeviceType=IOS
 BaseProfileName=IOS
 +CVars=sg.PostProcessQuality=3
 +CVars=sg.ShadowQuality=3
++CVars=r.Mobile.EyeAdaptation=0
 
 [AppleTV4K DeviceProfile]
 DeviceType=IOS
