Add missing Mali G77 device profile support

diff --git a/Engine/Config/BaseDeviceProfiles.ini b/Engine/Config/BaseDeviceProfiles.ini
--- a/Engine/Config/BaseDeviceProfiles.ini
+++ b/Engine/Config/BaseDeviceProfiles.ini
@@ -75,6 +75,9 @@
 +DeviceProfileNameAndTypes=Android_Mali_G72_Vulkan,Android
 +DeviceProfileNameAndTypes=Android_Mali_G76,Android
 +DeviceProfileNameAndTypes=Android_Mali_G76_Vulkan,Android
++DeviceProfileNameAndTypes=Android_Mali_G77,Android
++DeviceProfileNameAndTypes=Android_Mali_G77_Vulkan,Android
++DeviceProfileNameAndTypes=Android_Vulkan_SM5,Android
 +DeviceProfileNameAndTypes=Android_PowerVR_G6xxx,Android
 +DeviceProfileNameAndTypes=Android_PowerVR_GT7xxx,Android
 +DeviceProfileNameAndTypes=Android_PowerVR_GE8xxx,Android
@@ -410,7 +413,7 @@ MatchProfile=(Profile="Android_Adreno4xx",Match=((SourceType=SRC_GpuFamily,Compa
 
 ; Example profiles for Vulkan SM5
 ;+MatchProfile=(Profile="Android_Vulkan_SM5",Match=((SourceType=SRC_GpuFamily,CompareType=CMP_Regex,MatchString="Adreno \\(TM\\) 6[0-9][0-9]"),(SourceType=SRC_AndroidVersion, CompareType=CMP_Regex,MatchString="([0-9]+).*"),(SourceType=SRC_PreviousRegexMatch,CompareType=CMP_GreaterEqual,MatchString="9"),(SourceType=SRC_VulkanAvailable,CompareType=CMP_Equal,MatchString="true")))
-;+MatchProfile=(Profile="Android_Vulkan_SM5",Match=((SourceType=SRC_GpuFamily,CompareType=CMP_Regex,MatchString="^Mali\\-G76"),(SourceType=SRC_VulkanAvailable,CompareType=CMP_Equal,MatchString="true")))
+;+MatchProfile=(Profile="Android_Vulkan_SM5",Match=((SourceType=SRC_GpuFamily,CompareType=CMP_Regex,MatchString="^Mali\\-G7[67]"),(SourceType=SRC_VulkanAvailable,CompareType=CMP_Equal,MatchString="true")))
 
 +MatchProfile=(Profile="Android_Adreno5xx_Low",Match=((SourceType=SRC_GpuFamily,CompareType=CMP_Regex,MatchString="Adreno \\(TM\\) 5[0-1][0-9]")))
 +MatchProfile=(Profile="Android_Adreno5xx",Match=((SourceType=SRC_GpuFamily,CompareType=CMP_Regex,MatchString="Adreno \\(TM\\) 5[0-9][0-9]")))
@@ -424,6 +427,8 @@ MatchProfile=(Profile="Android_Adreno4xx",Match=((SourceType=SRC_GpuFamily,Compa
 +MatchProfile=(Profile="Android_Mali_G72",Match=((SourceType=SRC_GpuFamily,CompareType=CMP_Regex,MatchString="^Mali\\-G72")))
 +MatchProfile=(Profile="Android_Mali_G76_Vulkan",Match=((SourceType=SRC_GpuFamily,CompareType=CMP_Regex,MatchString="^Mali\\-G76"),(SourceType=SRC_VulkanAvailable,CompareType=CMP_Equal,MatchString="true")))
 +MatchProfile=(Profile="Android_Mali_G76",Match=((SourceType=SRC_GpuFamily,CompareType=CMP_Regex,MatchString="^Mali\\-G76")))
++MatchProfile=(Profile="Android_Mali_G77_Vulkan",Match=((SourceType=SRC_GpuFamily,CompareType=CMP_Regex,MatchString="^Mali\\-G77"),(SourceType=SRC_VulkanAvailable,CompareType=CMP_Equal,MatchString="true")))
++MatchProfile=(Profile="Android_Mali_G77",Match=((SourceType=SRC_GpuFamily,CompareType=CMP_Regex,MatchString="^Mali\\-G77")))
 +MatchProfile=(Profile="Android_PowerVR_G6xxx",Match=((SourceType=SRC_GpuFamily,CompareType=CMP_Regex,MatchString="PowerVR Rogue G6[0-9]+")))
 +MatchProfile=(Profile="Android_PowerVR_GT7xxx",Match=((SourceType=SRC_GpuFamily,CompareType=CMP_Regex,MatchString="PowerVR Rogue GT7[0-9]+")))
 +MatchProfile=(Profile="Android_PowerVR_GE8xxx",Match=((SourceType=SRC_GpuFamily,CompareType=CMP_Regex,MatchString="PowerVR Rogue GE8[0-9]+")))
@@ -608,6 +613,16 @@ BaseProfileName=Android_Mali_G76
 ; enable Vulkan
 +CVars=r.Android.DisableVulkanSupport=0
 
+[Android_Mali_G77 DeviceProfile]
+DeviceType=Android
+BaseProfileName=Android_High
+
+[Android_Mali_G77_Vulkan DeviceProfile]
+DeviceType=Android
+BaseProfileName=Android_Mali_G77
+; enable Vulkan
++CVars=r.Android.DisableVulkanSupport=0
+
 [Android_TegraK1 DeviceProfile]
 DeviceType=Android
 BaseProfileName=Android_High
