re-implment fix for startup crash that is currently used in 12.10: rely on hotfix currently live to reenable submix effects during startup, rather than setting the ini settings locally.

#localization none
#tests none#review [REDACTED_HANDLE], [REDACTED_HANDLE]

diff --git a/Engine/Config/ConsoleVariables.ini b/Engine/Config/ConsoleVariables.ini
--- a/Engine/Config/ConsoleVariables.ini
+++ b/Engine/Config/ConsoleVariables.ini
@@ -87,4 +87,4 @@ p.chaos.AllowCreatePhysxBodies=1
 fx.SkipVectorVMBackendOptimizations=1
 
 ; Comment this out to turn off audio submix effects.
-au.BypassAllSubmixEffects=0
\ No newline at end of file
+;au.BypassAllSubmixEffects=0
\ No newline at end of file
