Selected text is discernable in the Output log with default settings.

# Reverted default setting for the selected text background color to what it was before CL 3902112

#Editgrating CL 11188140 from Dev-Editor

diff --git a/Engine/Config/BaseEditorPerProjectUserSettings.ini b/Engine/Config/BaseEditorPerProjectUserSettings.ini
--- a/Engine/Config/BaseEditorPerProjectUserSettings.ini
+++ b/Engine/Config/BaseEditorPerProjectUserSettings.ini
@@ -90,7 +90,7 @@ InactiveSelectionColor=(R=0.25f,G=0.25f,B=0.25f,A=1.000000)
 SelectorColor=(R=0.701f,G=0.225f,B=0.003f,A=1.000000)
 
 LogBackgroundColor=(R=0.015996,G=0.015996,B=0.015996,A=1.000000)
-LogSelectionBackgroundColor=(R=0.008023,G=0.008023,B=0.008023,A=1.000000)
+LogSelectionBackgroundColor=(R=0.132868,G=0.132868,B=0.132868,A=1.000000)
 LogNormalColor=(R=0.72,G=0.72,B=0.72,A=1.000000)
 LogCommandColor=(R=0.033105,G=0.723055,B=0.033105,A=1.000000)
 LogWarningColor=(R=0.921875,G=0.691406,B=0.000000,A=1.000000)
