Fix a crash of Histogram CS on IOS A8 devices.

Fix some legacy tooltips for eyeadaptation.

Change the disable eyeadaptation condition for low-end ios devices to be consistent with android devices using r.EyeAdaptationQuality instead of r.Mobile.EyeAdaptation.

diff --git a/Engine/Config/BaseDeviceProfiles.ini b/Engine/Config/BaseDeviceProfiles.ini
--- a/Engine/Config/BaseDeviceProfiles.ini
+++ b/Engine/Config/BaseDeviceProfiles.ini
@@ -170,28 +170,28 @@ DeviceType=IOS
 BaseProfileName=IOS
 +CVars=r.RenderTargetSwitchWorkaround=1
 +CVars=fx.NiagaraAllowComputeShaders=0
-+CVars=r.Mobile.EyeAdaptation=0
++CVars=r.EyeAdaptationQuality=0
 
 [iPad3 DeviceProfile]
 DeviceType=IOS
 BaseProfileName=IOS
 +CVars=r.RenderTargetSwitchWorkaround=1
 +CVars=fx.NiagaraAllowComputeShaders=0
-+CVars=r.Mobile.EyeAdaptation=0
++CVars=r.EyeAdaptationQuality=0
 
 [iPad4 DeviceProfile]
 DeviceType=IOS
 BaseProfileName=IOS
 +CVars=r.RenderTargetSwitchWorkaround=1
 +CVars=fx.NiagaraAllowComputeShaders=0
-+CVars=r.Mobile.EyeAdaptation=0
++CVars=r.EyeAdaptationQuality=0
 
 [iPadAir DeviceProfile]
 DeviceType=IOS
 BaseProfileName=IOS
 +CVars=sg.PostProcessQuality=1
 +CVars=fx.NiagaraAllowComputeShaders=0
-+CVars=r.Mobile.EyeAdaptation=0
++CVars=r.EyeAdaptationQuality=0
 
 [iPadAir2 DeviceProfile]
 DeviceType=IOS
@@ -199,27 +199,27 @@ BaseProfileName=IOS
 +CVars=sg.PostProcessQuality=3
 +CVars=sg.ShadowQuality=2
 +CVars=fx.NiagaraAllowComputeShaders=0
-+CVars=r.Mobile.EyeAdaptation=0
++CVars=r.EyeAdaptationQuality=0
 
 [iPadMini DeviceProfile]
 DeviceType=IOS
 BaseProfileName=IOS
 +CVars=r.RenderTargetSwitchWorkaround=1
 +CVars=fx.NiagaraAllowComputeShaders=0
-+CVars=r.Mobile.EyeAdaptation=0
++CVars=r.EyeAdaptationQuality=0
 
 [iPadMini2 DeviceProfile]
 DeviceType=IOS
 BaseProfileName=iPadAir
 +CVars=fx.NiagaraAllowComputeShaders=0
-+CVars=r.Mobile.EyeAdaptation=0
++CVars=r.EyeAdaptationQuality=0
 
 [iPadMini4 DeviceProfile]
 DeviceType=IOS
 BaseProfileName=iPadMini2
 +CVars=sg.PostProcessQuality=2
 +CVars=fx.NiagaraAllowComputeShaders=0
-+CVars=r.Mobile.EyeAdaptation=0
++CVars=r.EyeAdaptationQuality=0
 
 [iPhone5S DeviceProfile]
 DeviceType=IOS
@@ -227,21 +227,21 @@ BaseProfileName=IOS
 +CVars=r.MobileContentScaleFactor=2
 +CVars=sg.PostProcessQuality=3
 +CVars=fx.NiagaraAllowComputeShaders=0
-+CVars=r.Mobile.EyeAdaptation=0
++CVars=r.EyeAdaptationQuality=0
 
 [iPodTouch6 DeviceProfile]
 DeviceType=IOS
 BaseProfileName=iPhone6
 +CVars=r.MobileContentScaleFactor=2
 +CVars=fx.NiagaraAllowComputeShaders=0
-+CVars=r.Mobile.EyeAdaptation=0
++CVars=r.EyeAdaptationQuality=0
 
 [iPodTouch7 DeviceProfile]
 DeviceType=IOS
 BaseProfileName=iPhone7
 +CVars=r.MobileContentScaleFactor=2
 +CVars=fx.NiagaraAllowComputeShaders=0
-+CVars=r.Mobile.EyeAdaptation=0
++CVars=r.EyeAdaptationQuality=0
 
 [iPhone6 DeviceProfile]
 DeviceType=IOS
@@ -249,7 +249,7 @@ BaseProfileName=IOS
 +CVars=r.MobileContentScaleFactor=2
 +CVars=sg.PostProcessQuality=3
 +CVars=fx.NiagaraAllowComputeShaders=0
-+CVars=r.Mobile.EyeAdaptation=0
++CVars=r.EyeAdaptationQuality=0
 
 [iPhone6Plus DeviceProfile]
 DeviceType=IOS
@@ -353,7 +353,7 @@ DeviceType=IOS
 BaseProfileName=IOS
 +CVars=sg.PostProcessQuality=3
 +CVars=sg.ShadowQuality=3
-+CVars=r.Mobile.EyeAdaptation=0
++CVars=r.EyeAdaptationQuality=0
 
 [AppleTV4K DeviceProfile]
 DeviceType=IOS
diff --git a/Engine/Shaders/Private/PostProcessMobile.usf b/Engine/Shaders/Private/PostProcessMobile.usf
--- a/Engine/Shaders/Private/PostProcessMobile.usf
+++ b/Engine/Shaders/Private/PostProcessMobile.usf
@@ -1387,8 +1387,19 @@ void Histogram_MainCS(
 				}
 
 				// Accumulate the weight to the nearby history buckets.
+#if IOS // The IOS A8 and lower devices seems don't support using a float4 array as a two dimension array, separate the operations to two steps.
+				float4 Histogram0 = float4(0.0f, 0.0f, 0.0f, 0.0f);
+				float4 Histogram1 = float4(0.0f, 0.0f, 0.0f, 0.0f);
+
+				Histogram0[Bucket0 % 4] = Weight0 * ScreenWeight;
+				Histogram1[Bucket1 % 4] = Weight1 * ScreenWeight;
+
+				SharedHistogram[(Bucket0 / 4) * ThreadGroupSize + GroupIndex] += Histogram0;
+				SharedHistogram[(Bucket1 / 4) * ThreadGroupSize + GroupIndex] += Histogram1;
+#else
 				SharedHistogram[(Bucket0 / 4) * ThreadGroupSize + GroupIndex][Bucket0 % 4] += Weight0 * ScreenWeight;
 				SharedHistogram[(Bucket1 / 4) * ThreadGroupSize + GroupIndex][Bucket1 % 4] += Weight1 * ScreenWeight;
+#endif
 			}
 		}
 	}
diff --git a/Engine/Source/Runtime/Engine/Classes/Engine/RendererSettings.h b/Engine/Source/Runtime/Engine/Classes/Engine/RendererSettings.h
--- a/Engine/Source/Runtime/Engine/Classes/Engine/RendererSettings.h
+++ b/Engine/Source/Runtime/Engine/Classes/Engine/RendererSettings.h
@@ -130,9 +130,9 @@ namespace EAutoExposureMethodUI
 {
 	enum Type
 	{
-		/** Not supported on mobile, requires compute shader to construct 64 bin histogram */
+		/** requires compute shader to construct 64 bin histogram */
 		AEM_Histogram  UMETA(DisplayName = "Auto Exposure Histogram"),
-		/** Not supported on mobile, faster method that computes single value by downsampling */
+		/** faster method that computes single value by downsampling */
 		AEM_Basic      UMETA(DisplayName = "Auto Exposure Basic"),
 		/** Uses camera settings. */
 		AEM_Manual   UMETA(DisplayName = "Manual"),
diff --git a/Engine/Source/Runtime/Engine/Classes/Engine/Scene.h b/Engine/Source/Runtime/Engine/Classes/Engine/Scene.h
--- a/Engine/Source/Runtime/Engine/Classes/Engine/Scene.h
+++ b/Engine/Source/Runtime/Engine/Classes/Engine/Scene.h
@@ -43,9 +43,9 @@ enum EAntiAliasingMethod
 UENUM()
 enum EAutoExposureMethod
 {
-	/** Not supported on mobile, requires compute shader to construct 64 bin histogram */
+	/** requires compute shader to construct 64 bin histogram */
 	AEM_Histogram  UMETA(DisplayName = "Auto Exposure Histogram"),
-	/** Not supported on mobile, faster method that computes single value by downsampling */
+	/** faster method that computes single value by downsampling */
 	AEM_Basic      UMETA(DisplayName = "Auto Exposure Basic"),
 	/** Uses camera settings. */
 	AEM_Manual   UMETA(DisplayName = "Manual"),
