UE-90862 HoloLens packaging fails in GetTextureBuildSettings()
-Fixed ini so HoloLens device profile can be found and used.

diff --git a/Engine/Config/BaseDeviceProfiles.ini b/Engine/Config/BaseDeviceProfiles.ini
--- a/Engine/Config/BaseDeviceProfiles.ini
+++ b/Engine/Config/BaseDeviceProfiles.ini
@@ -86,6 +86,7 @@
 +DeviceProfileNameAndTypes=Android_TegraK1,Android
 +DeviceProfileNameAndTypes=Lumin,Lumin
 +DeviceProfileNameAndTypes=Lumin_Desktop,Lumin
++DeviceProfileNameAndTypes=HoloLens,HoloLens
 
 [/Script/Engine.TextureLODSettings]
 @TextureLODGroups=Group
