Re-enable GpuComputeNiagara for A11+ gpus, was wrongly disabled

diff --git a/Engine/Config/BaseDeviceProfiles.ini b/Engine/Config/BaseDeviceProfiles.ini
--- a/Engine/Config/BaseDeviceProfiles.ini
+++ b/Engine/Config/BaseDeviceProfiles.ini
@@ -296,11 +296,13 @@ BaseProfileName=iPhone6Plus
 DeviceType=IOS
 BaseProfileName=iPhone7
 +CVars=sg.ShadowQuality=3
++CVars=fx.NiagaraAllowComputeShaders=1
 
 [iPhone8Plus DeviceProfile]
 DeviceType=IOS
 BaseProfileName=iPhone7Plus
 +CVars=sg.ShadowQuality=3
++CVars=fx.NiagaraAllowComputeShaders=1
 
 [iPhoneX DeviceProfile]
 DeviceType=IOS
