Don't allow drivers older than 26.20.13031.15006 (from 19-11-2) on AMD for DX12 because of PSO creation issues

diff --git a/Engine/Config/BaseHardware.ini b/Engine/Config/BaseHardware.ini
--- a/Engine/Config/BaseHardware.ini
+++ b/Engine/Config/BaseHardware.ini
@@ -66,7 +66,7 @@ SuggestedDriverVersion="19.11.1"
 ; 2017 drivers and older contribute heavily to swapchain creation errors.
 +Blacklist=(DriverVersion="<=22.19.662.4", Reason="Older drivers known to cause crashes with integrated laptop graphics on at least R5 and R6")
 
-+Blacklist=(DriverVersion="==26.20.13031.15006", Reason="Driver crashes creating PSOs", RHI="D3D12")
++Blacklist=(DriverVersion="<=26.20.13031.15006", Reason="Driver crashes creating PSOs", RHI="D3D12")
 
 ;-----------------------------------------------------------------------------------------------------------------
 
