Disable Vulkan support on devices with Mali-G72 GPUs on Android 8 and older versions. It can't create a PSO with a compute shader that uses texel buffer.

diff --git a/Engine/Config/BaseDeviceProfiles.ini b/Engine/Config/BaseDeviceProfiles.ini
--- a/Engine/Config/BaseDeviceProfiles.ini
+++ b/Engine/Config/BaseDeviceProfiles.ini
@@ -426,7 +426,7 @@ MatchProfile=(Profile="Android_Adreno4xx",Match=((SourceType=SRC_GpuFamily,Compa
 +MatchProfile=(Profile="Android_Mali_T7xx",Match=((SourceType=SRC_GpuFamily,CompareType=CMP_Regex,MatchString="^Mali\\-T7")))
 +MatchProfile=(Profile="Android_Mali_T8xx",Match=((SourceType=SRC_GpuFamily,CompareType=CMP_Regex,MatchString="^Mali\\-T8")))
 +MatchProfile=(Profile="Android_Mali_G71",Match=((SourceType=SRC_GpuFamily,CompareType=CMP_Regex,MatchString="^Mali\\-G71")))
-+MatchProfile=(Profile="Android_Mali_G72_Vulkan",Match=((SourceType=SRC_GpuFamily,CompareType=CMP_Regex,MatchString="^Mali\\-G72"),(SourceType=SRC_VulkanAvailable,CompareType=CMP_Equal,MatchString="true")))
++MatchProfile=(Profile="Android_Mali_G72_Vulkan",Match=((SourceType=SRC_GpuFamily,CompareType=CMP_Regex,MatchString="^Mali\\-G72"),(SourceType=SRC_AndroidVersion, CompareType=CMP_Regex,MatchString="([0-9]+).*"),(SourceType=SRC_PreviousRegexMatch,CompareType=CMP_GreaterEqual,MatchString="9"),(SourceType=SRC_VulkanAvailable,CompareType=CMP_Equal,MatchString="true")))
 +MatchProfile=(Profile="Android_Mali_G72",Match=((SourceType=SRC_GpuFamily,CompareType=CMP_Regex,MatchString="^Mali\\-G72")))
 +MatchProfile=(Profile="Android_Mali_G76_Vulkan",Match=((SourceType=SRC_GpuFamily,CompareType=CMP_Regex,MatchString="^Mali\\-G76"),(SourceType=SRC_VulkanAvailable,CompareType=CMP_Equal,MatchString="true")))
 +MatchProfile=(Profile="Android_Mali_G76",Match=((SourceType=SRC_GpuFamily,CompareType=CMP_Regex,MatchString="^Mali\\-G76")))
@@ -603,7 +603,7 @@ BaseProfileName=Android_High
 [Android_Mali_G72_Vulkan DeviceProfile]
 DeviceType=Android
 BaseProfileName=Android_Mali_G72
-; enable Vulkan
+; enable Vulkan on Android 9 and up, older versions crash on creating PSO with a compute shader that uses texel_buffer (eye adaptation)
 +CVars=r.Android.DisableVulkanSupport=0
 
 [Android_Mali_G76 DeviceProfile]
