Disable early_fragment_tests usage for Mali-T8xx GPUs as it causes rendering artifacts

diff --git a/Engine/Config/BaseDeviceProfiles.ini b/Engine/Config/BaseDeviceProfiles.ini
--- a/Engine/Config/BaseDeviceProfiles.ini
+++ b/Engine/Config/BaseDeviceProfiles.ini
@@ -591,6 +591,8 @@ BaseProfileName=Android_Low
 DeviceType=Android
 BaseProfileName=Android_Mid
 +CVars=r.Android.MaliT8Bug=1
+; using early_fragment_tests in a fragment shader does not work correctly on this device
++CVars=r.Android.DisableEarlyFragmentTests=1
 
 [Android_Mali_G71 DeviceProfile]
 DeviceType=Android
