Adding profile for TVOS that inherits AppleTV settings. Probably been broken for a while but hidden :(

diff --git a/Engine/Config/BaseDeviceProfiles.ini b/Engine/Config/BaseDeviceProfiles.ini
--- a/Engine/Config/BaseDeviceProfiles.ini
+++ b/Engine/Config/BaseDeviceProfiles.ini
@@ -46,6 +46,7 @@
 +DeviceProfileNameAndTypes=iPadPro3_129,IOS
 +DeviceProfileNameAndTypes=AppleTV,IOS
 +DeviceProfileNameAndTypes=AppleTV4K,IOS
++DeviceProfileNameAndTypes=TVOS,AppleTV
 +DeviceProfileNameAndTypes=Mac,Mac
 +DeviceProfileNameAndTypes=MacClient,Mac
 +DeviceProfileNameAndTypes=MacNoEditor,Mac
