Tweak some culture display name substitutes.
#rnx

diff --git a/Engine/Config/BaseEngine.ini b/Engine/Config/BaseEngine.ini
--- a/Engine/Config/BaseEngine.ini
+++ b/Engine/Config/BaseEngine.ini
@@ -909,6 +909,11 @@ gc.UseDisregardForGCOnDedicatedServers=False
 
 [Internationalization]
 +LocalizationPaths=../../../Engine/Content/Localization/Engine
++CultureDisplayNameSubstitutes="Taiwan;Chinese Taipei"
++CultureDisplayNameSubstitutes="ja;\u53f0\u6e7e;\u30c1\u30e3\u30a4\u30cb\u30fc\u30ba\u30bf\u30a4\u30da\u30a4"
++CultureDisplayNameSubstitutes="\u53f0\u6e7e;\u4e2d\u534e\u53f0\u5317"
++CultureDisplayNameSubstitutes="\u53f0\u7063;\u4e2d\u83ef\u53f0\u5317"
++CultureDisplayNameSubstitutes="\ub300\ub9cc;\uc911\ud654 \ud0c0\uc774\ubca0\uc774"
 
 [Audio]
 CommonAudioPoolSize=0
