Vulkan support for unknown Mali and Adreno devices running Android 10 or later

diff --git a/Engine/Config/BaseDeviceProfiles.ini b/Engine/Config/BaseDeviceProfiles.ini
--- a/Engine/Config/BaseDeviceProfiles.ini
+++ b/Engine/Config/BaseDeviceProfiles.ini
@@ -86,6 +86,7 @@
 +DeviceProfileNameAndTypes=Android_PowerVR_GM9xxx,Android
 +DeviceProfileNameAndTypes=Android_PowerVR_GM9xxx_Vulkan,Android
 +DeviceProfileNameAndTypes=Android_TegraK1,Android
++DeviceProfileNameAndTypes=Android_Unknown_Vulkan,Android
 +DeviceProfileNameAndTypes=Lumin,Lumin
 +DeviceProfileNameAndTypes=Lumin_Desktop,Lumin
 +DeviceProfileNameAndTypes=HoloLens,HoloLens
@@ -398,15 +399,15 @@ BaseProfileName=IPadPro
 ; These rules are checked in turn and stop at the profile that matches all criteria
 [/Script/AndroidDeviceProfileSelector.AndroidDeviceProfileMatchingRules]
 MatchProfile=(Profile="Android_Adreno4xx",Match=((SourceType=SRC_GpuFamily,CompareType=CMP_Regex,MatchString="Adreno \\(TM\\) 4[0-9][0-9]")))
-
++MatchProfile=(Profile="Android_Adreno5xx_Low",Match=((SourceType=SRC_GpuFamily,CompareType=CMP_Regex,MatchString="Adreno \\(TM\\) 5[0-1][0-9]")))
++MatchProfile=(Profile="Android_Adreno5xx",Match=((SourceType=SRC_GpuFamily,CompareType=CMP_Regex,MatchString="Adreno \\(TM\\) 5[0-9][0-9]")))
 ; Example profiles for Vulkan SM5
 ;+MatchProfile=(Profile="Android_Vulkan_SM5",Match=((SourceType=SRC_GpuFamily,CompareType=CMP_Regex,MatchString="Adreno \\(TM\\) 6[0-9][0-9]"),(SourceType=SRC_AndroidVersion, CompareType=CMP_Regex,MatchString="([0-9]+).*"),(SourceType=SRC_PreviousRegexMatch,CompareType=CMP_GreaterEqual,MatchString="9"),(SourceType=SRC_VulkanAvailable,CompareType=CMP_Equal,MatchString="true")))
 ;+MatchProfile=(Profile="Android_Vulkan_SM5",Match=((SourceType=SRC_GpuFamily,CompareType=CMP_Regex,MatchString="^Mali\\-G7[67]"),(SourceType=SRC_VulkanAvailable,CompareType=CMP_Equal,MatchString="true")))
-
-+MatchProfile=(Profile="Android_Adreno5xx_Low",Match=((SourceType=SRC_GpuFamily,CompareType=CMP_Regex,MatchString="Adreno \\(TM\\) 5[0-1][0-9]")))
-+MatchProfile=(Profile="Android_Adreno5xx",Match=((SourceType=SRC_GpuFamily,CompareType=CMP_Regex,MatchString="Adreno \\(TM\\) 5[0-9][0-9]")))
 +MatchProfile=(Profile="Android_Adreno6xx_Vulkan",Match=((SourceType=SRC_GpuFamily,CompareType=CMP_Regex,MatchString="Adreno \\(TM\\) 6[0-9][0-9]"),(SourceType=SRC_AndroidVersion, CompareType=CMP_Regex,MatchString="([0-9]+).*"),(SourceType=SRC_PreviousRegexMatch,CompareType=CMP_GreaterEqual,MatchString="9"),(SourceType=SRC_VulkanAvailable,CompareType=CMP_Equal,MatchString="true")))
 +MatchProfile=(Profile="Android_Adreno6xx",Match=((SourceType=SRC_GpuFamily,CompareType=CMP_Regex,MatchString="Adreno \\(TM\\) 6[0-9][0-9]")))
+; Enable Vulkan on any unknown Adreno running Android 10 or later
++MatchProfile=(Profile="Android_Unknown_Vulkan",Match=((SourceType=SRC_GpuFamily,CompareType=CMP_Regex,MatchString="Adreno"),(SourceType=SRC_AndroidVersion, CompareType=CMP_Regex,MatchString="([0-9]+).*"),(SourceType=SRC_PreviousRegexMatch,CompareType=CMP_GreaterEqual,MatchString="10"),(SourceType=SRC_VulkanAvailable,CompareType=CMP_Equal,MatchString="true")))
 +MatchProfile=(Profile="Android_Mali_T6xx",Match=((SourceType=SRC_GpuFamily,CompareType=CMP_Regex,MatchString="^Mali\\-T6")))
 +MatchProfile=(Profile="Android_Mali_T7xx",Match=((SourceType=SRC_GpuFamily,CompareType=CMP_Regex,MatchString="^Mali\\-T7")))
 +MatchProfile=(Profile="Android_Mali_T8xx",Match=((SourceType=SRC_GpuFamily,CompareType=CMP_Regex,MatchString="^Mali\\-T8")))
@@ -417,6 +418,8 @@ MatchProfile=(Profile="Android_Adreno4xx",Match=((SourceType=SRC_GpuFamily,Compa
 +MatchProfile=(Profile="Android_Mali_G76",Match=((SourceType=SRC_GpuFamily,CompareType=CMP_Regex,MatchString="^Mali\\-G76")))
 +MatchProfile=(Profile="Android_Mali_G77_Vulkan",Match=((SourceType=SRC_GpuFamily,CompareType=CMP_Regex,MatchString="^Mali\\-G77"),(SourceType=SRC_VulkanAvailable,CompareType=CMP_Equal,MatchString="true")))
 +MatchProfile=(Profile="Android_Mali_G77",Match=((SourceType=SRC_GpuFamily,CompareType=CMP_Regex,MatchString="^Mali\\-G77")))
+; Enable Vulkan on any unknown Mali running Android 10 or later
++MatchProfile=(Profile="Android_Unknown_Vulkan",Match=((SourceType=SRC_GpuFamily,CompareType=CMP_Regex,MatchString="^Mali"),(SourceType=SRC_AndroidVersion, CompareType=CMP_Regex,MatchString="([0-9]+).*"),(SourceType=SRC_PreviousRegexMatch,CompareType=CMP_GreaterEqual,MatchString="10"),(SourceType=SRC_VulkanAvailable,CompareType=CMP_Equal,MatchString="true")))
 +MatchProfile=(Profile="Android_PowerVR_G6xxx",Match=((SourceType=SRC_GpuFamily,CompareType=CMP_Regex,MatchString="PowerVR Rogue G6[0-9]+")))
 +MatchProfile=(Profile="Android_PowerVR_GT7xxx",Match=((SourceType=SRC_GpuFamily,CompareType=CMP_Regex,MatchString="PowerVR Rogue GT7[0-9]+")))
 +MatchProfile=(Profile="Android_PowerVR_GE8xxx",Match=((SourceType=SRC_GpuFamily,CompareType=CMP_Regex,MatchString="PowerVR Rogue GE8[0-9]+")))
@@ -621,6 +624,12 @@ BaseProfileName=Android_Mali_G77
 DeviceType=Android
 BaseProfileName=Android_High
 
+[Android_Unknown_Vulkan DeviceProfile]
+DeviceType=Android
+BaseProfileName=Android_High
+; enable Vulkan
++CVars=r.Android.DisableVulkanSupport=0
+
 [Lumin DeviceProfile]
 DeviceType=Lumin
 BaseProfileName=
