- Added a default ini setting to make Windows DDPI memory freezing ini parsing get correct value for raytracing

diff --git a/Engine/Config/BaseEngine.ini b/Engine/Config/BaseEngine.ini
--- a/Engine/Config/BaseEngine.ini
+++ b/Engine/Config/BaseEngine.ini
@@ -2192,6 +2192,7 @@ bEnable2DPhysics=false
 +TargetedRHIs=PCD3D_SM5
 MinimumOSVersion=MSOS_Vista
 bEnableRayTracing=true
+bTarget32Bit=false
 
 [/Script/LinuxTargetPlatform.LinuxTargetSettings]
 ; When neither -vulkan nor -opengl4 are passed on the commandline, the engine will default to the first targeted RHI.
