Take Recorder: Exclude VisibilityBasedAnimTickOption from being recorded

diff --git a/Engine/Config/BaseEditorSettings.ini b/Engine/Config/BaseEditorSettings.ini
--- a/Engine/Config/BaseEditorSettings.ini
+++ b/Engine/Config/BaseEditorSettings.ini
@@ -88,7 +88,7 @@ InteractorClass=Class'/Script/VREditor.VREditorInteractor'
 TeleporterClass=Class'/Script/VREditor.VREditorTeleporter'
 
 [/Script/TakeRecorder.TakeRecorderProjectSettings]
-+Settings=(DefaultTracks=((MatchingActorClass=/Script/CinematicCamera.CineCameraActor,DefaultPropertyTracks=((ComponentPath="CameraComponent",PropertyPath="CurrentFocalLength"),(ComponentPath="CameraComponent",PropertyPath="FocusSettings.ManualFocusDistance"),(ComponentPath="CameraComponent",PropertyPath="CurrentAperture"),(ComponentPath="CameraComponent",PropertyPath="FilmbackSettings.SensorWidth"),(ComponentPath="CameraComponent",PropertyPath="FilmbackSettings.SensorHeight")), ExcludePropertyTracks=((ComponentPath="CameraComponent",PropertyPath="FilmbackSettings.SensorAspectRatio"),(ComponentPath="CameraComponent",PropertyPath="AspectRatio"),(ComponentPath="CameraComponent",PropertyPath="FieldOfView"))),(MatchingActorClass=/Script/Engine.Light,DefaultPropertyTracks=((ComponentPath="LightComponent0",PropertyPath="Intensity"),(ComponentPath="LightComponent0",PropertyPath="LightColor")))))
++Settings=(DefaultTracks=((MatchingActorClass=/Script/CinematicCamera.CineCameraActor,DefaultPropertyTracks=((ComponentPath="CameraComponent",PropertyPath="CurrentFocalLength"),(ComponentPath="CameraComponent",PropertyPath="FocusSettings.ManualFocusDistance"),(ComponentPath="CameraComponent",PropertyPath="CurrentAperture"),(ComponentPath="CameraComponent",PropertyPath="FilmbackSettings.SensorWidth"),(ComponentPath="CameraComponent",PropertyPath="FilmbackSettings.SensorHeight")), ExcludePropertyTracks=((ComponentPath="CameraComponent",PropertyPath="FilmbackSettings.SensorAspectRatio"),(ComponentPath="CameraComponent",PropertyPath="AspectRatio"),(ComponentPath="CameraComponent",PropertyPath="FieldOfView"))),(MatchingActorClass=/Script/Engine.Light,DefaultPropertyTracks=((ComponentPath="LightComponent0",PropertyPath="Intensity"),(ComponentPath="LightComponent0",PropertyPath="LightColor"))),(MatchingActorClass=/Script/Engine.SkeletalMeshActor,ExcludePropertyTracks=((ComponentPath="SkeletalMeshComponent0",PropertyPath="VisibilityBasedAnimTickOption"))))
 
 [DefaultTakeMetaData TakeMetaData]
 Duration=(SubFrame=0.0)
