Enable Async Texture Compilation by default

diff --git a/Engine/Config/BaseEditorPerProjectUserSettings.ini b/Engine/Config/BaseEditorPerProjectUserSettings.ini
--- a/Engine/Config/BaseEditorPerProjectUserSettings.ini
+++ b/Engine/Config/BaseEditorPerProjectUserSettings.ini
@@ -83,12 +83,6 @@ bEnableUserEditorLayoutManagement=True
 ColorVisionDeficiencyPreviewType=NormalVision
 ColorVisionDeficiencySeverity=3
 
-
-SelectionColor=(R=0.728f,G=0.364f,B=0.003f,A=1.000000)
-PressedSelectionColor=(R=0.701f,G=0.225f,B=0.003f,A=1.000000)
-InactiveSelectionColor=(R=0.25f,G=0.25f,B=0.25f,A=1.000000)
-SelectorColor=(R=0.701f,G=0.225f,B=0.003f,A=1.000000)
-
 LogBackgroundColor=(R=0.015996,G=0.015996,B=0.015996,A=1.000000)
 LogSelectionBackgroundColor=(R=0.132868,G=0.132868,B=0.132868,A=1.000000)
 LogNormalColor=(R=0.72,G=0.72,B=0.72,A=1.000000)
@@ -153,6 +147,9 @@ bBehaviorTreeEditor=False
 bEnableFindAndReplaceReferences =False
 bExampleLayersAndBlends = True
 
+; Enabled by default now
+bEnableAsyncTextureCompilation=True
+
 [/Script/UnrealEd.EditorLoadingSavingSettings]
 ; True if we should automatically load a default level at start up
 LoadLevelAtStartup=ProjectDefault
