Delete old editor news feed style assets. The news feed has been deleted for years.

diff --git a/Engine/Config/BaseEditorSettings.ini b/Engine/Config/BaseEditorSettings.ini
--- a/Engine/Config/BaseEditorSettings.ini
+++ b/Engine/Config/BaseEditorSettings.ini
@@ -16,13 +16,6 @@ bSendUnattendedBugReports=True
 [/Script/UnrealEd.AnalyticsPrivacySettings]
 bSendUsageData=True
 
-[/Script/NewsFeed.NewsFeedSettings]
-CdnSourceUrl="http://cdn.unrealengine.com/newsfeed/[[Culture]]/index.json"
-LocalSourcePath="Editor/NewsFeed/[[Culture]]/"
-Source=NEWSFEED_Cdn
-MaxItemsToShow=10
-ShowOnlyUnreadItems=True
-
 [/Script/UnrealEd.ContentBrowserSettings]
 ; The number of objects to load at once in the Content Browser before displaying a warning about loading many assets
 NumObjectsToLoadBeforeWarning=20
