Added default blueprint graph event nodes for gameplay cue notifies.

diff --git a/Engine/Config/BaseEditorPerProjectUserSettings.ini b/Engine/Config/BaseEditorPerProjectUserSettings.ini
--- a/Engine/Config/BaseEditorPerProjectUserSettings.ini
+++ b/Engine/Config/BaseEditorPerProjectUserSettings.ini
@@ -807,6 +807,11 @@ DefaultBrushRadius=128
 +Node=(TargetClass=ActorComponent TargetEvent="ReceiveTick")
 +Node=(TargetClass=GameplayAbility TargetEvent="K2_ActivateAbility")
 +Node=(TargetClass=GameplayAbility TargetEvent="K2_OnEndAbility")
++Node=(TargetClass=GameplayCueNotify_BurstLatent TargetEvent="OnBurst")
++Node=(TargetClass=GameplayCueNotify_Looping TargetEvent="OnApplication")
++Node=(TargetClass=GameplayCueNotify_Looping TargetEvent="OnLoopingStart")
++Node=(TargetClass=GameplayCueNotify_Looping TargetEvent="OnRecurring")
++Node=(TargetClass=GameplayCueNotify_Looping TargetEvent="OnRemoval")
 +Node=(TargetClass=UserWidget TargetEvent="PreConstruct")
 +Node=(TargetClass=UserWidget TargetEvent="Construct")
 +Node=(TargetClass=UserWidget TargetEvent="Tick")
