Remove odd special case VT pool sizes from default engine config
Just leave the default size as 64MB, rely on projects to specify more precise sizes if desired

diff --git a/Engine/Config/BaseEngine.ini b/Engine/Config/BaseEngine.ini
--- a/Engine/Config/BaseEngine.ini
+++ b/Engine/Config/BaseEngine.ini
@@ -2515,5 +2515,3 @@ DefaultChaosSolverActorClass=/Script/ChaosSolverEngine.ChaosSolverActor
 ; "r.VT.ListPhysicalPool" list in-depth details on the allocated physical pools
 ; The upper allocation limit for a pool is caluclated by matching a size and format from the list below or falling back to the default
 DefaultSizeInMegabyte=64
-+Pools=(SizeInMegabyte=64, MinTileSize=0, MaxTileSize=9999, Formats=(PF_DXT1, PF_DXT5))
-+Pools=(SizeInMegabyte=512, MinTileSize=0, MaxTileSize=9999, Formats=(PF_A32B32G32R32F))
