Fix SceneOutliner ini property name 'bHideLevelInstanceHierarchy'

diff --git a/Engine/Config/BaseEditorPerProjectUserSettings.ini b/Engine/Config/BaseEditorPerProjectUserSettings.ini
--- a/Engine/Config/BaseEditorPerProjectUserSettings.ini
+++ b/Engine/Config/BaseEditorPerProjectUserSettings.ini
@@ -123,8 +123,8 @@ bShowOnlyActorsInCurrentLevel=False
 bShowOnlySelectedActors=False
 ; Hide actor components
 bHideActorComponents=True
-; Hide foundation hierarchy
-bHideFoundationHierarchy=False
+; Hide level instance hierarchy
+bHideLevelInstanceHierarchy=False
 
 
 [/Script/UnrealEd.ClassViewerSettings]
