Enable async static mesh compilation by default

diff --git a/Engine/Config/BaseEditorPerProjectUserSettings.ini b/Engine/Config/BaseEditorPerProjectUserSettings.ini
--- a/Engine/Config/BaseEditorPerProjectUserSettings.ini
+++ b/Engine/Config/BaseEditorPerProjectUserSettings.ini
@@ -154,6 +154,7 @@ bExampleLayersAndBlends = True
 
 ; Enabled by default now
 bEnableAsyncTextureCompilation=True
+bEnableAsyncStaticMeshCompilation=True
 
 [/Script/UnrealEd.EditorLoadingSavingSettings]
 ; True if we should automatically load a default level at start up
