Blueprints: Add the promotable operator to the favorites menu by defalut

#rnx

diff --git a/Engine/Config/BaseEditor.ini b/Engine/Config/BaseEditor.ini
--- a/Engine/Config/BaseEditor.ini
+++ b/Engine/Config/BaseEditor.ini
@@ -188,6 +188,17 @@ DefaultProfile="StarterProfile"
 +StarterProfile=(NodeName="/Script/BlueprintGraph.K2Node_CallFunction", FieldName="/Script/Engine.Actor:K2_DestroyActor")
 +StarterProfile=(NodeName="/Script/BlueprintGraph.K2Node_CallFunction", FieldName="/Script/Engine.GameplayStatics:GetAllActorsOfClass")
 +StarterProfile=(NodeName="/Script/BlueprintGraph.K2Node_SpawnActorFromClass")
++StarterProfile=(NodeName="/Script/BlueprintGraph.K2Node_PromotableOperator")
++StarterProfile=(FieldName="/Script/TimeManagement.TimeManagementBlueprintLibrary:Add_FrameNumberInteger",NodeName="/Script/BlueprintGraph.K2Node_PromotableOperator")
++StarterProfile=(FieldName="/Script/TimeManagement.TimeManagementBlueprintLibrary:Subtract_FrameNumberInteger",NodeName="/Script/BlueprintGraph.K2Node_PromotableOperator")
++StarterProfile=(FieldName="/Script/GameplayTags.BlueprintGameplayTagLibrary:EqualEqual_GameplayTagContainer",NodeName="/Script/BlueprintGraph.K2Node_PromotableOperator")
++StarterProfile=(FieldName="/Script/TimeManagement.TimeManagementBlueprintLibrary:Divide_FrameNumberInteger",NodeName="/Script/BlueprintGraph.K2Node_PromotableOperator")
++StarterProfile=(FieldName="/Script/Engine.KismetMathLibrary:Greater_TimespanTimespan",NodeName="/Script/BlueprintGraph.K2Node_PromotableOperator")
++StarterProfile=(FieldName="/Script/Engine.KismetMathLibrary:GreaterEqual_TimespanTimespan",NodeName="/Script/BlueprintGraph.K2Node_PromotableOperator")
++StarterProfile=(FieldName="/Script/Engine.KismetMathLibrary:Less_TimespanTimespan",NodeName="/Script/BlueprintGraph.K2Node_PromotableOperator")
++StarterProfile=(FieldName="/Script/Engine.KismetMathLibrary:LessEqual_TimespanTimespan",NodeName="/Script/BlueprintGraph.K2Node_PromotableOperator")
++StarterProfile=(FieldName="/Script/GameplayTags.BlueprintGameplayTagLibrary:NotEqual_GameplayTagContainer",NodeName="/Script/BlueprintGraph.K2Node_PromotableOperator")
++StarterProfile=(FieldName="/Script/TimeManagement.TimeManagementBlueprintLibrary:Multiply_SecondsFrameRate",NodeName="/Script/BlueprintGraph.K2Node_PromotableOperator")
 
 [BlueprintEditor.Menu]
 +BlueprintHiddenFields="/Script/Engine.Actor:K2_AttachRootComponentTo"
