Cameras: Fix struct names in core redirectors.

diff --git a/Engine/Config/BaseEngine.ini b/Engine/Config/BaseEngine.ini
--- a/Engine/Config/BaseEngine.ini
+++ b/Engine/Config/BaseEngine.ini
@@ -834,9 +834,9 @@ NearClipPlane=10.0
 +EnumRedirects=(OldName="EOscillatorWaveform",NewName="/Script/GameplayCameras.EOscillatorWaveform")
 +EnumRedirects=(OldName="EInitialOscillatorOffset",NewName="/Script/GameplayCameras.EInitialOscillatorOffset")
 ; Note: the following three redirects may break blueprints because BreakStruct and MakeStruct nodes don't support redirectors.
-+StructRedirects=(OldName="FFOscillator",NewName="/Script/GameplayCameras.FFOscillator")
-+StructRedirects=(OldName="FROscillator",NewName="/Script/GameplayCameras.FROscillator")
-+StructRedirects=(OldName="FVOscillator",NewName="/Script/GameplayCameras.FVOscillator")
++StructRedirects=(OldName="FOscillator",NewName="/Script/GameplayCameras.FOscillator")
++StructRedirects=(OldName="ROscillator",NewName="/Script/GameplayCameras.ROscillator")
++StructRedirects=(OldName="VOscillator",NewName="/Script/GameplayCameras.VOscillator")
 +ClassRedirects=(OldName="CameraShake",NewName="/Script/GameplayCameras.MatineeCameraShake")
 +ClassRedirects=(OldName="MatineeCameraShake",NewName="/Script/GameplayCameras.MatineeCameraShake")
 +ClassRedirects=(OldName="MatineeCameraShakePattern",NewName="/Script/GameplayCameras.MatineeCameraShakePattern")
