Corrected default PipelineStack path for Interchange.

diff --git a/Engine/Config/BaseEngine.ini b/Engine/Config/BaseEngine.ini
--- a/Engine/Config/BaseEngine.ini
+++ b/Engine/Config/BaseEngine.ini
@@ -2557,4 +2557,4 @@ bDisableActionLimit=false
 ActionLimit=100
 
 [/Script/InterchangeEngine.InterchangeProjectSettings]
-PipelineStack=/Script/InterchangePipelinePlugin.InterchangeGenericMeshPipeline
\ No newline at end of file
+PipelineStack=/Script/InterchangePipelines.InterchangeGenericMeshPipeline
