Added new DDC profile -DDC=Cold

This option allows the engine to run against an empty DDC folder whilst preserving any existing DDC folders

diff --git a/Engine/Config/BaseEngine.ini b/Engine/Config/BaseEngine.ini
--- a/Engine/Config/BaseEngine.ini
+++ b/Engine/Config/BaseEngine.ini
@@ -1800,6 +1800,15 @@ Boot=(Type=Boot, Filename="%GAMEDIR%DerivedDataCache/Boot.ddc", MaxCacheSize=512
 Local=(Type=FileSystem, ReadOnly=false, Clean=false, Flush=false, PurgeTransient=true, DeleteUnused=true, UnusedFileAge=34, FoldersToClean=-1, Path=%ENGINEDIR%DerivedDataCache, EnvPathOverride=UE-LocalDataCachePath)
 Pak=(Type=ReadPak, Filename="%GAMEDIR%DerivedDataCache/DDC.ddp")
 
+[Cold]
+MinimumDaysToKeepFile=0
+Root=(Type=KeyLength, Length=120, Inner=AsyncPut)
+AsyncPut=(Type=AsyncPut, Inner=Hierarchy)
+Hierarchy=(Type=Hierarchical, Inner=Boot, Inner=Local)
+Boot=(Type=Boot, Filename="%GAMEDIR%DerivedDataCache/Boot.ddc", MaxCacheSize=512)
+Local=(Type=FileSystem, ReadOnly=false, Clean=true, Flush=false, PurgeTransient=true, DeleteUnused=true, UnusedFileAge=1, FoldersToClean=-1, Path="%GAMEDIR%ColdDerivedDataCache\")
+
+
 [CreatePak]
 MinimumDaysToKeepFile=7
 Root=(Type=KeyLength, Length=120, Inner=AsyncPut)
