Added new DDC profile -DDC=Cold

This option allows the engine to run against an empty DDC folder whilst preserving any existing DDC folders
Fixed settings so clean=false flush=true

diff --git a/Engine/Config/BaseEngine.ini b/Engine/Config/BaseEngine.ini
--- a/Engine/Config/BaseEngine.ini
+++ b/Engine/Config/BaseEngine.ini
@@ -1806,7 +1806,7 @@ Root=(Type=KeyLength, Length=120, Inner=AsyncPut)
 AsyncPut=(Type=AsyncPut, Inner=Hierarchy)
 Hierarchy=(Type=Hierarchical, Inner=Boot, Inner=Local)
 Boot=(Type=Boot, Filename="%GAMEDIR%DerivedDataCache/Boot.ddc", MaxCacheSize=512)
-Local=(Type=FileSystem, ReadOnly=false, Clean=true, Flush=false, PurgeTransient=true, DeleteUnused=true, UnusedFileAge=1, FoldersToClean=-1, Path="%GAMEDIR%ColdDerivedDataCache\")
+Local=(Type=FileSystem, ReadOnly=false, Clean=false, Flush=true, PurgeTransient=true, DeleteUnused=true, UnusedFileAge=1, FoldersToClean=-1, Path="%GAMEDIR%ColdDerivedDataCache\")
 
 
 [CreatePak]
