Blacklist Nvidia driver 462.40

diff --git a/Engine/Config/BaseHardware.ini b/Engine/Config/BaseHardware.ini
--- a/Engine/Config/BaseHardware.ini
+++ b/Engine/Config/BaseHardware.ini
@@ -58,6 +58,8 @@ SuggestedDriverVersion="441.20"
 +Blacklist=(DriverVersion="==382.05", Reason="This driver version has many known stability issues")
 +Blacklist=(DriverVersion="==388.57", Reason="This driver version has many known stability issues")
 
+; Driver with known stability issues
++Blacklist=(DriverVersion="==461.40", Reason="This driver version has known stability issues")
 ;-----------------------------------------------------------------------------------------------------------------
 
 [GPU_AMD]
