Reverted the removal/redirection of CalculateDirection

There are cases in use where the (unused) self pin has been connected that fail to compile

diff --git a/Engine/Config/BaseEngine.ini b/Engine/Config/BaseEngine.ini
--- a/Engine/Config/BaseEngine.ini
+++ b/Engine/Config/BaseEngine.ini
@@ -868,7 +868,6 @@ NearClipPlane=10.0
 +ClassRedirects=(OldName="/Script/MeshDescription.MeshDescriptionBulkDataWrapper",NewName="/Script/Engine.StaticMeshDescriptionBulkData")
 +StructRedirects=(OldName="/Script/Engine.FilePath",NewName="/Script/CoreUObject.FilePath")
 +StructRedirects=(OldName="/Script/Engine.DirectoryPath",NewName="/Script/CoreUObject.DirectoryPath")
-+FunctionRedirects=(OldName="CalculateDirection",NewName="/Script/AnimGraphRuntime.KismetAnimationLibrary.CalculateDirection")
 
 [CoreUObject.Metadata]
 ; Note: UnrealHeaderTool should be rerun after making changes to MetadataRedirects to catch any keys specified in class headers
diff --git a/Engine/Source/Runtime/Engine/Classes/Animation/AnimInstance.h b/Engine/Source/Runtime/Engine/Classes/Animation/AnimInstance.h
--- a/Engine/Source/Runtime/Engine/Classes/Animation/AnimInstance.h
+++ b/Engine/Source/Runtime/Engine/Classes/Animation/AnimInstance.h
@@ -1053,6 +1053,7 @@ public:
 	void ClearMorphTargets();
 
 	UE_DEPRECATED(5.0, "Please use UKismetAnimationLibrary::CalculateDirection instead")
+	UFUNCTION(BlueprintCallable, Category="Animation")
 	float CalculateDirection(const FVector& Velocity, const FRotator& BaseRotation) const;
 
 	//--- AI communication start ---//
