Fixup blendspace warnings with exposed blendspace pins

diff --git a/Engine/Config/BaseEngine.ini b/Engine/Config/BaseEngine.ini
--- a/Engine/Config/BaseEngine.ini
+++ b/Engine/Config/BaseEngine.ini
@@ -913,6 +913,7 @@ NearClipPlane=10.0
 +PropertyRedirects=(OldName="AnimationBlueprintLibrary.GetTimeAtFrame.AnimationSequence",NewName="AnimationSequenceBase")
 +PropertyRedirects=(OldName="AnimationBlueprintLibrary.IsValidTime.AnimationSequence",NewName="AnimationSequenceBase")
 +PropertyRedirects=(OldName="AnimationBlueprintLibrary.FindBonePathToRoot.AnimationSequence",NewName="AnimationSequenceBase")
++ClassRedirects=(OldName="/Script/Engine.BlendSpaceBase",NewName="/Script/Engine.BlendSpace")
 
 [CoreUObject.Metadata]
 ; Note: UnrealHeaderTool should be rerun after making changes to MetadataRedirects to catch any keys specified in class headers
