Enable GPU channel in Trace.ChannelPresets Default preset.

diff --git a/Engine/Config/BaseEngine.ini b/Engine/Config/BaseEngine.ini
--- a/Engine/Config/BaseEngine.ini
+++ b/Engine/Config/BaseEngine.ini
@@ -2043,7 +2043,7 @@ DrawColor=(R=38,G=38,B=38,A=64)
 DrawColor=(R=140,G=255,B=0,A=164)
 
 [Trace.ChannelPresets]
-Default=cpu,frame,log,bookmark
+Default=gpu,cpu,frame,log,bookmark
 Rendering=gpu,cpu,frame,log,bookmark
 Memory=memtag,memalloc,callstack,module
 
