Removed reference to UE4 in comments

diff --git a/Engine/Config/BaseDeviceProfiles.ini b/Engine/Config/BaseDeviceProfiles.ini
--- a/Engine/Config/BaseDeviceProfiles.ini
+++ b/Engine/Config/BaseDeviceProfiles.ini
@@ -730,7 +730,6 @@ BaseProfileName=Android_Adreno6xx
 
 ;----------------
 ; PowerVR devices
-; some PowerVR G6xxx devices are known to have problems with ES 3.1 on UE4
 
 [Android_PowerVR_G6xxx DeviceProfile]
 DeviceType=Android
