Enabled r.DumpShaderDebugInfo=2 by default. This automatically saves shader compilation diagnostic / repro data on compilation failures.

diff --git a/Engine/Config/ConsoleVariables.ini b/Engine/Config/ConsoleVariables.ini
--- a/Engine/Config/ConsoleVariables.ini
+++ b/Engine/Config/ConsoleVariables.ini
@@ -34,9 +34,13 @@
 
 ; Uncomment to get detailed logs on shader compiles and the opportunity to retry on errors
 ;r.ShaderDevelopmentMode=1
-; Uncomment to dump shaders in the Saved folder (1 dump all, 2 dump on compilation failure only, 3 dump on compilation failure or warnings)
-; Warning: leaving this on for a while will fill your hard drive with many small files and folders
-;r.DumpShaderDebugInfo=1
+
+; Dump shaders in the Saved folder
+;  Mode 1: dump all. WARNING: leaving this on for a while will fill your hard drive with many small files and folders.
+;  Mode 2: dump on compilation failure only (default).
+;  Mode 3: dump on compilation failure or warnings.
+r.DumpShaderDebugInfo=2
+
 ; When this is enabled, dumped shader paths will get collapsed (in the cases where paths are longer than the OS's max)
 ;r.DumpShaderDebugShortNames=1
 ; When this is enabled, when dumping shaders an additional file to use with ShaderCompilerWorker -direct mode will be generated
