Fix Max Metal Shader Standard to Target default value

diff --git a/Engine/Config/BaseEngine.ini b/Engine/Config/BaseEngine.ini
--- a/Engine/Config/BaseEngine.ini
+++ b/Engine/Config/BaseEngine.ini
@@ -2231,7 +2231,7 @@ bAutomaticSigning=false
 UseFastIntrinsics=False
 ForceFloats=False
 EnableMathOptimisations=True
-MaxShaderLanguageVersion=3
+MaxShaderLanguageVersion=5
 bDisableMotionData=False
 bEnableAdvertisingIdentifier=True
 
