[NO NON-PII COMMIT MESSAGE]

diff --git a/Engine/Config/BaseEngine.ini b/Engine/Config/BaseEngine.ini
--- a/Engine/Config/BaseEngine.ini
+++ b/Engine/Config/BaseEngine.ini
@@ -1958,7 +1958,7 @@ AspectRatioAxisConstraint=AspectRatio_MaintainXFOV
 ; These are long package paths
 ;+BlacklistPackagePathScanFilters="/Game/__ExternalActors__"
 ; Default sub content paths that do not get scanned on each mount by the asset registry when searching for all assets
-+BlacklistContentSubPathScanFilters="__ExternalActors__"
+;+BlacklistContentSubPathScanFilters="__ExternalActors__"
 ; Fill out this list with asset registry tags that are not needed in cooked builds and will be stripped during cook.
 ; You can use a * wildcard to refer to either all classes or all tags for the class.
 ; If your game does not need any asset registry tags, simply add (Class=*,Tag=*). Be warned that this this may make some engine systems fail to work since they might need ParentClass or GeneratedClass on Blueprints.
