Fully deprecate SetBrushSize.

diff --git a/Engine/Config/BaseEngine.ini b/Engine/Config/BaseEngine.ini
--- a/Engine/Config/BaseEngine.ini
+++ b/Engine/Config/BaseEngine.ini
@@ -922,6 +922,7 @@ NearClipPlane=10.0
 +PropertyRedirects=(OldName="AnimationBlueprintLibrary.IsValidTime.AnimationSequence",NewName="AnimationSequenceBase")
 +PropertyRedirects=(OldName="AnimationBlueprintLibrary.FindBonePathToRoot.AnimationSequence",NewName="AnimationSequenceBase")
 +ClassRedirects=(OldName="/Script/Engine.BlendSpaceBase",NewName="/Script/Engine.BlendSpace")
++FunctionRedirects=(OldName="Image.SetBrushSize",NewName="SetDesiredSizeOverride")
 
 [CoreUObject.Metadata]
 ; Note: UnrealHeaderTool should be rerun after making changes to MetadataRedirects to catch any keys specified in class headers
diff --git a/Engine/Source/Runtime/UMG/Public/Components/Image.h b/Engine/Source/Runtime/UMG/Public/Components/Image.h
--- a/Engine/Source/Runtime/UMG/Public/Components/Image.h
+++ b/Engine/Source/Runtime/UMG/Public/Components/Image.h
@@ -75,7 +75,7 @@ public:
 	void SetOpacity(float InOpacity);
 
 	/**  */
-	UFUNCTION(BlueprintCallable, Category = "Appearance", meta = (DeprecatedFunction, DeprecationMessage = "Use SetDesiredSizeOverride() instead"))
+	UE_DEPRECATED(5.0, "Deprecated. Use SetDesiredSizeOverride instead.")
 	void SetBrushSize(FVector2D DesiredSize);
 
 	UFUNCTION(BlueprintCallable, Category = "Appearance")
