undo commented blacklist from cl 15714549 until other content browser filtering optimization are merged in UE5

diff --git a/Engine/Config/BaseEditor.ini b/Engine/Config/BaseEditor.ini
--- a/Engine/Config/BaseEditor.ini
+++ b/Engine/Config/BaseEditor.ini
@@ -25,7 +25,7 @@ ColorFontAtlasSize=512
 ; Full Package Path being blacklisted from the Asset view
 ;+BlacklistAssetPaths="/Game/__ExternalActors__"
 ; Content Sub Path being blacklisted from the Asset view for each mount
-;+BlacklistContentSubPaths="__ExternalActors__"
++BlacklistContentSubPaths="__ExternalActors__"
 
 [/Script/UnrealEd.ThumbnailManager]
 +RenderableThumbnailTypes=(ClassNeedingThumbnailName="/Script/CoreUObject.Class",RendererClassName="/Script/UnrealEd.ClassThumbnailRenderer")
