remove -asynccompression from default pak options

there appears to be a race bug in UnrealPak
since Pak is being phased out for iostore, just turn it off for now

diff --git a/Engine/Config/BaseGame.ini b/Engine/Config/BaseGame.ini
--- a/Engine/Config/BaseGame.ini
+++ b/Engine/Config/BaseGame.ini
@@ -88,7 +88,7 @@ UsePakFile=True
 bUseIoStore=True
 bCompressed=True
 PackageCompressionFormat=Oodle
-PackageAdditionalCompressionOptions=-compressionblocksize=256KB -asynccompression
+PackageAdditionalCompressionOptions=-compressionblocksize=256KB
 PackageCompressionMethod=Kraken
 PackageCompressionLevel_Distribution=7
 PackageCompressionLevel_TestShipping=5
