Removed external actor subpath from AssetTools's BlacklistContentSubPaths as it is no longer needed. ContentBrowserDataUtils::PathPassesAttributeFilter now handles the case of ExternalActors.
This saves a considerable amount of time when UContentBrowserAssetDataSource::CreateAssetFilter builds a compiled filter.

diff --git a/Engine/Config/BaseEditor.ini b/Engine/Config/BaseEditor.ini
--- a/Engine/Config/BaseEditor.ini
+++ b/Engine/Config/BaseEditor.ini
@@ -25,7 +25,7 @@ ColorFontAtlasSize=512
 ; Full Package Path being blacklisted from the Asset view
 ;+BlacklistAssetPaths="/Game/__ExternalActors__"
 ; Content Sub Path being blacklisted from the Asset view for each mount
-+BlacklistContentSubPaths="__ExternalActors__"
+;+BlacklistContentSubPaths=
 
 [/Script/UnrealEd.ThumbnailManager]
 +RenderableThumbnailTypes=(ClassNeedingThumbnailName="/Script/CoreUObject.Class",RendererClassName="/Script/UnrealEd.ClassThumbnailRenderer")
