BaseEngine.ini nop comment TextureFormatOodle options

diff --git a/Engine/Config/BaseEngine.ini b/Engine/Config/BaseEngine.ini
--- a/Engine/Config/BaseEngine.ini
+++ b/Engine/Config/BaseEngine.ini
@@ -2804,6 +2804,10 @@ CompressEffortLevel_Editor=20
 CompressEffortLevel_NoEditor=30
 ; bDebugColor shows the BC1-7 format used
 bDebugColor=False
+; bDebugDump writes out all images cooked
+bDebugDump=False
+; LogVerbosity 0,1, or 2
+LogVerbosity=0
 ; Lambda controls the quality vs size tradeoff for RDO
 GlobalLambdaMultiplier=1.000000
 DefaultRDOLambda=30
