[Memory] - Add rhi.DumpResourceMemory command to memreport.

diff --git a/Engine/Config/BaseEngine.ini b/Engine/Config/BaseEngine.ini
--- a/Engine/Config/BaseEngine.ini
+++ b/Engine/Config/BaseEngine.ini
@@ -1676,6 +1676,7 @@ bRecompileOnLoad=true
 +Cmd="LogCountedInstances"
 +Cmd="obj list -alphasort"
 +Cmd="rhi.DumpMemory"
++Cmd="rhi.DumpResourceMemory"
 +Cmd="LogOutStatLevels"
 +Cmd="ListSpawnedActors"
 
