Enabled anisotropic render passes for high, epic and cine quality mode.
I have verified that the shader generated does not increase on QAGame.
ShooterGame compiled and but not really tested as there is a runtime assert.

diff --git a/Engine/Config/BaseScalability.ini b/Engine/Config/BaseScalability.ini
--- a/Engine/Config/BaseScalability.ini
+++ b/Engine/Config/BaseScalability.ini
@@ -391,7 +391,6 @@ r.SceneColorFormat=3
 r.DetailMode=0
 r.TranslucencyVolumeBlur=0
 r.MaterialQualityLevel=0 ; Low quality
-r.AnisotropicMaterials=0
 r.SSS.Scale=0
 r.SSS.SampleSet=0
 r.SSS.Quality=0
@@ -424,7 +423,6 @@ r.SceneColorFormat=3
 r.DetailMode=1
 r.TranslucencyVolumeBlur=0
 r.MaterialQualityLevel=2 ; Medium quality
-r.AnisotropicMaterials=0
 r.SSS.Scale=0.75
 r.SSS.SampleSet=0
 r.SSS.Quality=0
@@ -457,7 +455,6 @@ r.SceneColorFormat=3
 r.DetailMode=1
 r.TranslucencyVolumeBlur=1
 r.MaterialQualityLevel=1 ; High quality
-r.AnisotropicMaterials=0
 r.SSS.Scale=1
 r.SSS.SampleSet=1
 r.SSS.Quality=-1
@@ -497,7 +494,6 @@ r.SceneColorFormat=4
 r.DetailMode=2
 r.TranslucencyVolumeBlur=1
 r.MaterialQualityLevel=1 ; High quality
-r.AnisotropicMaterials=0
 r.SSS.Scale=1
 r.SSS.SampleSet=2
 r.SSS.Quality=1
@@ -537,7 +533,6 @@ r.SceneColorFormat=4
 r.DetailMode=2
 r.TranslucencyVolumeBlur=1
 r.MaterialQualityLevel=1 ; High quality
-r.AnisotropicMaterials=0
 r.SSS.Scale=1
 r.SSS.SampleSet=2
 r.SSS.Quality=1
@@ -607,28 +602,33 @@ r.HairStrands.SkyLighting.IntegrationType=2
 r.HairStrands.SkyAO.SampleCount=4
 r.HairStrands.Visibility.MSAA.SamplePerPixel=1
 r.HairStrands.Interpolation.UseSingleGuide=1
+r.AnisotropicMaterials=0
 
 [ShadingQuality@1]
 r.HairStrands.SkyLighting.IntegrationType=2
 r.HairStrands.SkyAO.SampleCount=4
 r.HairStrands.Visibility.MSAA.SamplePerPixel=1
 r.HairStrands.Interpolation.UseSingleGuide=1
+r.AnisotropicMaterials=0
 
 [ShadingQuality@2]
 r.HairStrands.SkyLighting.IntegrationType=2
 r.HairStrands.SkyAO.SampleCount=4
 r.HairStrands.Visibility.MSAA.SamplePerPixel=4
 r.HairStrands.Interpolation.UseSingleGuide=1
+r.AnisotropicMaterials=1
 
 [ShadingQuality@3]
 r.HairStrands.SkyLighting.IntegrationType=2
 r.HairStrands.SkyAO.SampleCount=4
 r.HairStrands.Visibility.MSAA.SamplePerPixel=4
 r.HairStrands.Interpolation.UseSingleGuide=0
+r.AnisotropicMaterials=1
 
 [ShadingQuality@Cine]
 r.HairStrands.SkyLighting.IntegrationType=0
 r.HairStrands.SkyAO.SampleCount=8
 r.HairStrands.Visibility.MSAA.SamplePerPixel=8
 r.HairStrands.Interpolation.UseSingleGuide=0
+r.AnisotropicMaterials=1
 
