Matinee/Editor: Removed Matinee bindings/settings from various ini files

diff --git a/Engine/Config/BaseEditor.ini b/Engine/Config/BaseEditor.ini
--- a/Engine/Config/BaseEditor.ini
+++ b/Engine/Config/BaseEditor.ini
@@ -211,51 +211,8 @@ DefaultProfile="StarterProfile"
 HitProxySize=5
 
 [/Script/UnrealEd.UnrealEdOptions]
-+EditorCategories=(Name="Matinee")
 +EditorCategories=(Name="CurveEditor")
 
-;Matinee Commands
-+EditorCommands=(Parent="Matinee", CommandName="Matinee_PlayForward", Description="Matinee_PlayForward_Desc", ExecCommand="Matinee Play")
-+EditorCommands=(Parent="Matinee", CommandName="Matinee_PlayReverse", Description="Matinee_PlayReverse_Desc", ExecCommand="Matinee PlayReverse")
-+EditorCommands=(Parent="Matinee", CommandName="Matinee_Stop", Description="Matinee_Stop_Desc", ExecCommand="Matinee Stop")
-+EditorCommands=(Parent="Matinee", CommandName="Matinee_TogglePlayPause", Description="Matinee_TogglePlayPause_Desc", ExecCommand="Matinee TogglePlayPause")
-+EditorCommands=(Parent="Matinee", CommandName="Matinee_ZoomIn", Description="Matinee_ZoomIn_Desc", ExecCommand="Matinee ZoomIn")
-+EditorCommands=(Parent="Matinee", CommandName="Matinee_ZoomInAlt", Description="Matinee_ZoomInAlt_Desc", ExecCommand="Matinee ZoomIn")
-+EditorCommands=(Parent="Matinee", CommandName="Matinee_ZoomOut", Description="Matinee_ZoomOut_Desc", ExecCommand="Matinee ZoomOut")
-+EditorCommands=(Parent="Matinee", CommandName="Matinee_ZoomOutAlt", Description="Matinee_ZoomOutAlt_Desc", ExecCommand="Matinee ZoomOut")
-+EditorCommands=(Parent="Matinee", CommandName="Matinee_Undo", Description="Matinee_Undo_Desc", ExecCommand="Matinee Undo")
-+EditorCommands=(Parent="Matinee", CommandName="Matinee_Redo", Description="Matinee_Redo_Desc", ExecCommand="Matinee Redo")
-+EditorCommands=(Parent="Matinee", CommandName="Matinee_Copy", Description="Matinee_Copy_Desc", ExecCommand="Matinee Copy")
-+EditorCommands=(Parent="Matinee", CommandName="Matinee_Cut", Description="Matinee_Cut_Desc", ExecCommand="Matinee Cut")
-+EditorCommands=(Parent="Matinee", CommandName="Matinee_Paste", Description="Matinee_Paste_Desc", ExecCommand="Matinee Paste")
-+EditorCommands=(Parent="Matinee", CommandName="Matinee_DeleteSelection", Description="Matinee_DeleteSelection_Desc", ExecCommand="Matinee DeleteSelection")
-+EditorCommands=(Parent="Matinee", CommandName="Matinee_MarkInSection", Description="Matinee_MarkInSection_Desc", ExecCommand="Matinee MarkInSection")
-+EditorCommands=(Parent="Matinee", CommandName="Matinee_MarkOutSection", Description="Matinee_MarkOutSection_Desc", ExecCommand="Matinee MarkOutSection")
-+EditorCommands=(Parent="Matinee", CommandName="Matinee_IncrementPosition", Description="Matinee_IncrementPosition_Desc", ExecCommand="Matinee IncrementPosition")
-+EditorCommands=(Parent="Matinee", CommandName="Matinee_DecrementPosition", Description="Matinee_DecrementPosition_Desc", ExecCommand="Matinee DecrementPosition")
-+EditorCommands=(Parent="Matinee", CommandName="Matinee_MoveToNextKey", Description="Matinee_MoveToNextKey_Desc", ExecCommand="Matinee MoveToNextKey")
-+EditorCommands=(Parent="Matinee", CommandName="Matinee_MoveToPrevKey", Description="Matinee_MoveToPrevKey_Desc", ExecCommand="Matinee MoveToPrevKey")
-+EditorCommands=(Parent="Matinee", CommandName="Matinee_SplitAnimKey", Description="Matinee_SplitAnimKey_Desc", ExecCommand="Matinee SplitAnimKey")
-+EditorCommands=(Parent="Matinee", CommandName="Matinee_ToggleSnap", Description="Matinee_ToggleSnap_Desc", ExecCommand="Matinee ToggleSnap")
-+EditorCommands=(Parent="Matinee", CommandName="Matinee_MoveActiveUp", Description="Matinee_MoveActiveUp_Desc", ExecCommand="Matinee MoveActiveUp")
-+EditorCommands=(Parent="Matinee", CommandName="Matinee_MoveActiveDown", Description="Matinee_MoveActiveDown_Desc", ExecCommand="Matinee MoveActiveDown")
-+EditorCommands=(Parent="Matinee", CommandName="Matinee_AddKey", Description="Matinee_AddKey_Desc", ExecCommand="Matinee AddKey")
-+EditorCommands=(Parent="Matinee", CommandName="Matinee_DuplicateSelectedKeys", Description="Matinee_DuplicateSelectedKeys_Desc", ExecCommand="Matinee DuplicateSelectedKeys")
-+EditorCommands=(Parent="Matinee", CommandName="Matinee_CropAnimationBeginning", Description="Matinee_CropAnimationBeginning_Desc", ExecCommand="Matinee CropAnimationBeginning")
-+EditorCommands=(Parent="Matinee", CommandName="Matinee_CropAnimationEnd", Description="Matinee_CropAnimationEnd_Desc", ExecCommand="Matinee CropAnimationEnd")
-
-+EditorCommands=(Parent="Matinee", CommandName="Matinee_ViewFitSequence", Description="Matinee_ViewFitSequence_Desc", ExecCommand="Matinee ViewFitSequence")
-+EditorCommands=(Parent="Matinee", CommandName="Matinee_ViewFitToSelected", Description="Matinee_ViewFitToSelected_Desc", ExecCommand="Matinee ViewFitToSelected")
-+EditorCommands=(Parent="Matinee", CommandName="Matinee_ViewFitLoop", Description="Matinee_ViewFitLoop_Desc", ExecCommand="Matinee ViewFitLoop")
-+EditorCommands=(Parent="Matinee", CommandName="Matinee_ViewFitLoopSequence", Description="Matinee_ViewFitLoopSequence_Desc", ExecCommand="Matinee ViewFitLoopSequence")
-+EditorCommands=(Parent="Matinee", CommandName="Matinee_ViewEndOfTrack", Description="Matinee_ViewEndOfTrack_Desc", ExecCommand="Matinee ViewEndOfTrack")
-
-+EditorCommands=(Parent="Matinee", CommandName="Matinee_ChangeKeyInterpModeAUTO", Description="Matinee_ChangeKeyInterpModeAUTO_Desc", ExecCommand="Matinee ChangeKeyInterpModeAUTO")
-+EditorCommands=(Parent="Matinee", CommandName="Matinee_ChangeKeyInterpModeUSER", Description="Matinee_ChangeKeyInterpModeUSER_Desc", ExecCommand="Matinee ChangeKeyInterpModeUSER")
-+EditorCommands=(Parent="Matinee", CommandName="Matinee_ChangeKeyInterpModeBREAK", Description="Matinee_ChangeKeyInterpModeBREAK_Desc", ExecCommand="Matinee ChangeKeyInterpModeBREAK")
-+EditorCommands=(Parent="Matinee", CommandName="Matinee_ChangeKeyInterpModeLINEAR", Description="Matinee_ChangeKeyInterpModeLINEAR_Desc", ExecCommand="Matinee ChangeKeyInterpModeLINEAR")
-+EditorCommands=(Parent="Matinee", CommandName="Matinee_ChangeKeyInterpModeCONSTANT", Description="Matinee_ChangeKeyInterpModeCONSTANT_Desc", ExecCommand="Matinee ChangeKeyInterpModeCONSTANT")
-
 ;Curve Editor Commands
 +EditorCommands=(Parent="CurveEditor", CommandName="CurveEditor_ChangeInterpModeAUTO", Description="CurveEditor_ChangeInterpModeAUTO_Desc", ExecCommand="CurveEditor ChangeInterpModeAUTO")
 +EditorCommands=(Parent="CurveEditor", CommandName="CurveEditor_ChangeInterpModeUSER", Description="CurveEditor_ChangeInterpModeUSER_Desc", ExecCommand="CurveEditor ChangeInterpModeUSER")
diff --git a/Engine/Config/BaseEditorKeyBindings.ini b/Engine/Config/BaseEditorKeyBindings.ini
--- a/Engine/Config/BaseEditorKeyBindings.ini
+++ b/Engine/Config/BaseEditorKeyBindings.ini
@@ -1,44 +1,4 @@
 [/Script/UnrealEd.UnrealEdKeyBindings]
-+KeyBindings=(bCtrlDown=False,bAltDown=False,bShiftDown=False,Key="J",CommandName="Matinee_PlayReverse")
-+KeyBindings=(bCtrlDown=False,bAltDown=False,bShiftDown=False,Key="K",CommandName="Matinee_Stop")
-+KeyBindings=(bCtrlDown=False,bAltDown=False,bShiftDown=False,Key="L",CommandName="Matinee_PlayForward")
-+KeyBindings=(bCtrlDown=False,bAltDown=False,bShiftDown=False,Key="SpaceBar",CommandName="Matinee_TogglePlayPause")
-+KeyBindings=(bCtrlDown=False,bAltDown=False,bShiftDown=False,Key="Delete",CommandName="Matinee_DeleteSelection")
-+KeyBindings=(bCtrlDown=False,bAltDown=False,bShiftDown=False,Key="Add",CommandName="Matinee_ZoomIn")
-+KeyBindings=(bCtrlDown=False,bAltDown=False,bShiftDown=False,Key="Subtract",CommandName="Matinee_ZoomOut")
-+KeyBindings=(bCtrlDown=False,bAltDown=False,bShiftDown=False,Key="Equals",CommandName="Matinee_ZoomInAlt")
-+KeyBindings=(bCtrlDown=False,bAltDown=False,bShiftDown=False,Key="Underscore",CommandName="Matinee_ZoomOutAlt")
-+KeyBindings=(bCtrlDown=True,bAltDown=False,bShiftDown=False,Key="Z",CommandName="Matinee_Undo")
-+KeyBindings=(bCtrlDown=True,bAltDown=False,bShiftDown=False,Key="Y",CommandName="Matinee_Redo")
-+KeyBindings=(bCtrlDown=False,bAltDown=False,bShiftDown=False,Key="I",CommandName="Matinee_MarkInSection")
-+KeyBindings=(bCtrlDown=False,bAltDown=False,bShiftDown=False,Key="O",CommandName="Matinee_MarkOutSection")
-+KeyBindings=(bCtrlDown=False,bAltDown=False,bShiftDown=False,Key="Right",CommandName="Matinee_IncrementPosition")
-+KeyBindings=(bCtrlDown=False,bAltDown=False,bShiftDown=False,Key="Left",CommandName="Matinee_DecrementPosition")
-+KeyBindings=(bCtrlDown=False,bAltDown=False,bShiftDown=True,Key="Right",CommandName="Matinee_MoveToNextKey")
-+KeyBindings=(bCtrlDown=False,bAltDown=False,bShiftDown=True,Key="Left",CommandName="Matinee_MoveToPrevKey")
-+KeyBindings=(bCtrlDown=False,bAltDown=False,bShiftDown=False,Key="R",CommandName="Matinee_SplitAnimKey")
-+KeyBindings=(bCtrlDown=False,bAltDown=False,bShiftDown=False,Key="S",CommandName="Matinee_ToggleSnap")
-+KeyBindings=(bCtrlDown=False,bAltDown=False,bShiftDown=False,Key="Up",CommandName="Matinee_MoveActiveUp")
-+KeyBindings=(bCtrlDown=False,bAltDown=False,bShiftDown=False,Key="Down",CommandName="Matinee_MoveActiveDown")
-+KeyBindings=(bCtrlDown=False,bAltDown=False,bShiftDown=False,Key="Enter",CommandName="Matinee_AddKey")
-+KeyBindings=(bCtrlDown=True,bAltDown=False,bShiftDown=False,Key="W",CommandName="Matinee_DuplicateSelectedKeys")
-+KeyBindings=(bCtrlDown=True,bAltDown=False,bShiftDown=False,Key="I",CommandName="Matinee_CropAnimationBeginning")
-+KeyBindings=(bCtrlDown=True,bAltDown=False,bShiftDown=False,Key="O",CommandName="Matinee_CropAnimationEnd")
-+KeyBindings=(bCtrlDown=True,bAltDown=False,bShiftDown=False,Key="C",CommandName="Matinee_Copy")
-+KeyBindings=(bCtrlDown=True,bAltDown=False,bShiftDown=False,Key="X",CommandName="Matinee_Cut")
-+KeyBindings=(bCtrlDown=True,bAltDown=False,bShiftDown=False,Key="V",CommandName="Matinee_Paste")
-+KeyBindings=(bCtrlDown=False,bAltDown=False,bShiftDown=False,Key="A",CommandName="Matinee_ViewFitSequence")
-+KeyBindings=(bCtrlDown=False,bAltDown=False,bShiftDown=False,Key="F",CommandName="Matinee_ViewFitToSelected")
-+KeyBindings=(bCtrlDown=False,bAltDown=False,bShiftDown=True,Key="A",CommandName="Matinee_ViewFitLoop")
-+KeyBindings=(bCtrlDown=True,bAltDown=False,bShiftDown=False,Key="A",CommandName="Matinee_ViewFitLoopSequence")
-+KeyBindings=(bCtrlDown=False,bAltDown=False,bShiftDown=False,Key="E",CommandName="Matinee_ViewEndOfTrack")
-
-+KeyBindings=(bCtrlDown=False,bAltDown=False,bShiftDown=False,Key="one",CommandName="Matinee_ChangeKeyInterpModeAUTO")
-+KeyBindings=(bCtrlDown=False,bAltDown=False,bShiftDown=False,Key="two",CommandName="Matinee_ChangeKeyInterpModeUSER")
-+KeyBindings=(bCtrlDown=False,bAltDown=False,bShiftDown=False,Key="three",CommandName="Matinee_ChangeKeyInterpModeBREAK")
-+KeyBindings=(bCtrlDown=False,bAltDown=False,bShiftDown=False,Key="four",CommandName="Matinee_ChangeKeyInterpModeLINEAR")
-+KeyBindings=(bCtrlDown=False,bAltDown=False,bShiftDown=False,Key="five",CommandName="Matinee_ChangeKeyInterpModeCONSTANT")
-
 ; --- Curve Editor Bindings
 +KeyBindings=(bCtrlDown=False,bAltDown=False,bShiftDown=False,Key="one",CommandName="CurveEditor_ChangeInterpModeAUTO")
 +KeyBindings=(bCtrlDown=False,bAltDown=False,bShiftDown=False,Key="two",CommandName="CurveEditor_ChangeInterpModeUSER")
diff --git a/Engine/Config/BaseEngine.ini b/Engine/Config/BaseEngine.ini
--- a/Engine/Config/BaseEngine.ini
+++ b/Engine/Config/BaseEngine.ini
@@ -249,28 +249,6 @@ bEnableOnScreenDebugMessages=true
 DurationOfErrorsAndWarningsOnHUD=0
 NearClipPlane=10.0
 
-; Matinee redirects (starts-with match and replace)
-+MatineeTrackRedirects=(TargetClassName="/Script/Engine.CameraActor",OldFieldName="PostProcessSettings.EyeAdaptationLowPercent",NewFieldName="CameraComponent.PostProcessSettings.AutoExposureLowPercent")
-+MatineeTrackRedirects=(TargetClassName="/Script/Engine.CameraActor",OldFieldName="PostProcessSettings.EyeAdaptationHighPercent",NewFieldName="CameraComponent.PostProcessSettings.AutoExposureHighPercent")
-+MatineeTrackRedirects=(TargetClassName="/Script/Engine.CameraActor",OldFieldName="PostProcessSettings.EyeAdaptationMinBrightness",NewFieldName="CameraComponent.PostProcessSettings.AutoExposureMinBrightness")
-+MatineeTrackRedirects=(TargetClassName="/Script/Engine.CameraActor",OldFieldName="PostProcessSettings.EyeAdaptationMaxBrightness",NewFieldName="CameraComponent.PostProcessSettings.AutoExposureMaxBrightness")
-+MatineeTrackRedirects=(TargetClassName="/Script/Engine.CameraActor",OldFieldName="PostProcessSettings.EyeAdaptionSpeedDown",NewFieldName="CameraComponent.PostProcessSettings.AutoExposureSpeedDown")
-+MatineeTrackRedirects=(TargetClassName="/Script/Engine.CameraActor",OldFieldName="PostProcessSettings.EyeAdaptionSpeedUp",NewFieldName="CameraComponent.PostProcessSettings.AutoExposureSpeedUp")
-+MatineeTrackRedirects=(TargetClassName="/Script/Engine.CameraActor",OldFieldName="PostProcessSettings.ExposureOffset",NewFieldName="CameraComponent.PostProcessSettings.AutoExposureBias")
-
-+MatineeTrackRedirects=(TargetClassName="/Script/Engine.PostprocessVolume",OldFieldName="Settings.EyeAdaptationLowPercent",NewFieldName="Settings.AutoExposureLowPercent")
-+MatineeTrackRedirects=(TargetClassName="/Script/Engine.PostprocessVolume",OldFieldName="Settings.EyeAdaptationHighPercent",NewFieldName="Settings.AutoExposureHighPercent")
-+MatineeTrackRedirects=(TargetClassName="/Script/Engine.PostprocessVolume",OldFieldName="Settings.EyeAdaptationMinBrightness",NewFieldName="Settings.AutoExposureMinBrightness")
-+MatineeTrackRedirects=(TargetClassName="/Script/Engine.PostprocessVolume",OldFieldName="Settings.EyeAdaptationMaxBrightness",NewFieldName="Settings.AutoExposureMaxBrightness")
-+MatineeTrackRedirects=(TargetClassName="/Script/Engine.PostprocessVolume",OldFieldName="Settings.EyeAdaptionSpeedDown",NewFieldName="Settings.AutoExposureSpeedDown")
-+MatineeTrackRedirects=(TargetClassName="/Script/Engine.PostprocessVolume",OldFieldName="Settings.EyeAdaptionSpeedUp",NewFieldName="Settings.AutoExposureSpeedUp")
-+MatineeTrackRedirects=(TargetClassName="/Script/Engine.PostprocessVolume",OldFieldName="Settings.ExposureOffset",NewFieldName="Settings.AutoExposureBias")
-
-+MatineeTrackRedirects=(TargetClassName="/Script/Engine.CameraActor",OldFieldName="FOVAngle",NewFieldName="CameraComponent.FieldOfView")
-+MatineeTrackRedirects=(TargetClassName="/Script/Engine.CameraActor",OldFieldName="AspectRatio",NewFieldName="CameraComponent.AspectRatio")
-+MatineeTrackRedirects=(TargetClassName="/Script/Engine.CameraActor",OldFieldName="bConstrainAspectRatio",NewFieldName="CameraComponent.bConstrainAspectRatio")
-+MatineeTrackRedirects=(TargetClassName="/Script/Engine.CameraActor",OldFieldName="PostProcessSettings.",NewFieldName="CameraComponent.PostProcessSettings.")
-
 [CoreRedirects]
 
 ; All redirects prior to UE 4.10 are now in FCoreRedirects::RegisterNativeRedirects, new ones should be added to the bottom of this section
@@ -1388,7 +1366,6 @@ GameCommandLine=-log
 FOVAngle=90.000000
 GodMode=True
 UseAxisIndicator=True
-MatineeCurveDetail=0.1
 HeightMapExportClassName="TerrainHeightMapExporterTextT3D"
 bGroupingActive=true
 bCustomCameraAlignEmitter=true
