Removing warning text about EditorDomain.
#rnx

diff --git a/Engine/Config/BaseEditor.ini b/Engine/Config/BaseEditor.ini
--- a/Engine/Config/BaseEditor.ini
+++ b/Engine/Config/BaseEditor.ini
@@ -327,7 +327,6 @@ BatchCount=100
 MaxMemoryAllowance=16384
 
 [CookSettings]
-; EditorDomain is in prototype and currently is a significant deoptimization
 EditorDomainEnabled=false
 EditorDomainSaveUnversioned=false
 PackagesPerGC=0
