Removed references to UE4 in BaseEgnine.ini

diff --git a/Engine/Config/BaseEngine.ini b/Engine/Config/BaseEngine.ini
--- a/Engine/Config/BaseEngine.ini
+++ b/Engine/Config/BaseEngine.ini
@@ -2217,7 +2217,7 @@ Rendering=gpu,cpu,frame,log,bookmark
 [RemoteConfiguration]
 Enabled=false
 ConfigPathPrefix=\\epicgames.net\root\Home
-ConfigPathSuffix=UE4Cloud
+ConfigPathSuffix=UE5Cloud
 Timeout=1.0f
 +IniToLoad=EditorPerProjectUserSettings
 +IniToLoad=EditorKeyBindings
