Renamed UE5 reference to UE in BaseEngine.ini

diff --git a/Engine/Config/BaseEngine.ini b/Engine/Config/BaseEngine.ini
--- a/Engine/Config/BaseEngine.ini
+++ b/Engine/Config/BaseEngine.ini
@@ -2217,7 +2217,7 @@ Rendering=gpu,cpu,frame,log,bookmark
 [RemoteConfiguration]
 Enabled=false
 ConfigPathPrefix=\\epicgames.net\root\Home
-ConfigPathSuffix=UE5Cloud
+ConfigPathSuffix=UECloud
 Timeout=1.0f
 +IniToLoad=EditorPerProjectUserSettings
 +IniToLoad=EditorKeyBindings
