Uses Marcie_Grain_v3_128 as default engine film grain texture

diff --git a/Engine/Config/BaseEngine.ini b/Engine/Config/BaseEngine.ini
--- a/Engine/Config/BaseEngine.ini
+++ b/Engine/Config/BaseEngine.ini
@@ -110,7 +110,7 @@ DefaultBSPVertexTextureName=/Engine/EditorResources/BSPVertex.BSPVertex
 HighFrequencyNoiseTextureName=/Engine/EngineMaterials/Good64x64TilingNoiseHighFreq.Good64x64TilingNoiseHighFreq
 DefaultBokehTextureName=/Engine/EngineMaterials/DefaultBokeh.DefaultBokeh
 DefaultBloomKernelTextureName=/Engine/EngineMaterials/DefaultBloomKernel.DefaultBloomKernel
-DefaultFilmGrainTextureName=/Engine/EngineMaterials/Marcie_Grain_1024_M5_555_A0.Marcie_Grain_1024_M5_555_A0
+DefaultFilmGrainTextureName=/Engine/EngineResources/FilmGrains/Marcie_Grain_v3_128_M2_000.Marcie_Grain_v3_128_M2_000
 GeomMaterialName=/Engine/EngineDebugMaterials/GeomMaterial.GeomMaterial
 DebugMeshMaterialName=/Engine/EngineDebugMaterials/DebugMeshMaterial.DebugMeshMaterial
 EmissiveMeshMaterialName=/Engine/EngineMaterials/EmissiveMeshMaterial.EmissiveMeshMaterial
diff --git a/Engine/Source/Runtime/Renderer/Private/SceneRendering.cpp b/Engine/Source/Runtime/Renderer/Private/SceneRendering.cpp
--- a/Engine/Source/Runtime/Renderer/Private/SceneRendering.cpp
+++ b/Engine/Source/Runtime/Renderer/Private/SceneRendering.cpp
@@ -2355,8 +2355,8 @@ FSceneRenderer::FSceneRenderer(const FSceneViewFamily* InViewFamily,FHitProxyCon
 
 				if (FilmGrainTexture == GEngine->DefaultFilmGrainTexture)
 				{
-					// Decode Texture2D'/Engine/EngineMaterials/Marcie_Grain_1024_M5_555_A0.Marcie_Grain_1024_M5_555_A0'
-					ViewInfo->FinalPostProcessSettings.FilmGrainDecodeMultiply = (1.0f / 0.18f);
+					// Decode Texture2D'/Engine/EngineResources/FilmGrains/Marcie_Grain_v3_128_M2_000.Marcie_Grain_v3_128_M2_000'
+					ViewInfo->FinalPostProcessSettings.FilmGrainDecodeMultiply = 2.000f;
 					ViewInfo->FinalPostProcessSettings.FilmGrainDecodeAdd = 0.0f;
 				}
 			}
