#7890 : PR : Remove net.UseAdaptiveNetUpdateFrequency=0 from ConsoleVariables.ini

diff --git a/Engine/Config/ConsoleVariables.ini b/Engine/Config/ConsoleVariables.ini
--- a/Engine/Config/ConsoleVariables.ini
+++ b/Engine/Config/ConsoleVariables.ini
@@ -91,8 +91,6 @@ r.DumpShaderDebugInfo=2
 ;D3D12.EmitRgpFrameMarkers=1
 ;r.AsyncPipelineCompile=0
 
-net.UseAdaptiveNetUpdateFrequency=0
-
 ; 0 effectively disables physx, 1 allows physx to still run alongside ChaosPhysics
 p.chaos.AllowCreatePhysxBodies=1
 
