[memreport] - Adding useful GPU memory commands for memory profiling.

- High level summary of GPU memory attributed to rendering systems in UE5 (Lumen, Nanite, Shadows, etc.)
- Adds listing non-virtual textures (Always, fully resident textures are a common target area for optimization).
- Adds listing non-compressed textures (Bad for memory and for performance).
- List static and skeletal meshes sorted by resource size (Vertex and index buffers)

diff --git a/Engine/Config/BaseEngine.ini b/Engine/Config/BaseEngine.ini
--- a/Engine/Config/BaseEngine.ini
+++ b/Engine/Config/BaseEngine.ini
@@ -1672,6 +1672,30 @@ bRecompileOnLoad=true
 +Cmd="obj list -alphasort"
 +Cmd="rhi.DumpMemory"
 +Cmd="rhi.DumpResourceMemory"
++Cmd="rhi.dumpresourcememory summary name=Lumen"
++Cmd="rhi.dumpresourcememory summary name=Nanite"
++Cmd="rhi.dumpresourcememory summary name=Shadow"
++Cmd="rhi.dumpresourcememory summary name=DistanceFields"
++Cmd="rhi.dumpresourcememory summary name=GPUScene"
++Cmd="rhi.dumpresourcememory summary name=SkinCache"
++Cmd="rhi.dumpresourcememory summary name=TangentsRHIBuffer"
++Cmd="rhi.dumpresourcememory summary name=IndexBuffer"
++Cmd="rhi.dumpresourcememory summary name=VertexBuffer"
++Cmd="rhi.dumpresourcememory summary name=TexCoordRHIBuffer"
++Cmd="rhi.dumpresourcememory summary name=None"
++Cmd="rhi.dumpresourcememory summary Type=TextureCube"
++Cmd="rhi.dumpresourcememory summary name=CaptureCube"
++Cmd="rhi.dumpresourcememory summary name=MorphDeltasVB"
++Cmd="rhi.dumpresourcememory summary name=VertexIndicesVB"
++Cmd="rhi.dumpresourcememory summary name=VirtualTexture"
++Cmd="rhi.dumpresourcememory summary name=VirtualPhysicalTexture"
++Cmd="rhi.dumpresourcememory summary name=VolumetricFog"
++Cmd="rhi.dumpresourcememory summary name=Hair"
++Cmd="rhi.dumpresourcememory summary name=Subsurface"
++Cmd="obj list class=SkeletalMesh -resourcesizesort"
++Cmd="obj list class=StaticMesh -resourcesizesort"
++Cmd="listtextures nonvt"
++Cmd="listtextures uncompressed"
 +Cmd="LogOutStatLevels"
 +Cmd="ListSpawnedActors"
 
