Setups FFT bloom defaults

diff --git a/Engine/Config/BaseScalability.ini b/Engine/Config/BaseScalability.ini
--- a/Engine/Config/BaseScalability.ini
+++ b/Engine/Config/BaseScalability.ini
@@ -190,6 +190,7 @@ r.LensFlareQuality=0
 r.SceneColorFringeQuality=0
 r.EyeAdaptationQuality=2
 r.BloomQuality=4
+r.Bloom.HalfResolutionFFT=2
 r.FastBlurThreshold=0
 r.Upscale.Quality=1
 r.Tonemapper.GrainQuantization=0
@@ -209,6 +210,7 @@ r.LensFlareQuality=0
 r.SceneColorFringeQuality=0
 r.EyeAdaptationQuality=2
 r.BloomQuality=4
+r.Bloom.HalfResolutionFFT=2
 r.FastBlurThreshold=2
 r.Upscale.Quality=2
 r.Tonemapper.GrainQuantization=0
@@ -240,6 +242,7 @@ r.LensFlareQuality=2
 r.SceneColorFringeQuality=1
 r.EyeAdaptationQuality=2
 r.BloomQuality=5
+r.Bloom.HalfResolutionFFT=1
 r.FastBlurThreshold=3
 r.Upscale.Quality=2
 r.Tonemapper.GrainQuantization=1
@@ -273,6 +276,7 @@ r.LensFlareQuality=2
 r.SceneColorFringeQuality=1
 r.EyeAdaptationQuality=2
 r.BloomQuality=5
+r.Bloom.HalfResolutionFFT=1
 r.FastBlurThreshold=100
 r.Upscale.Quality=3
 r.Tonemapper.GrainQuantization=1
@@ -309,6 +313,7 @@ r.LensFlareQuality=3
 r.SceneColorFringeQuality=1
 r.EyeAdaptationQuality=2
 r.BloomQuality=5
+r.Bloom.HalfResolutionFFT=0
 r.FastBlurThreshold=100
 r.Upscale.Quality=3
 r.Tonemapper.GrainQuantization=1
