Adding a ZenOnlyDDC graph that is limited to the Zen backend alone.

diff --git a/Engine/Config/BaseEngine.ini b/Engine/Config/BaseEngine.ini
--- a/Engine/Config/BaseEngine.ini
+++ b/Engine/Config/BaseEngine.ini
@@ -2004,6 +2004,11 @@ Pak=(Type=ReadPak, Filename="%GAMEDIR%DerivedDataCache/DDC.ddp")
 ; Configure a Project Pak node. This is a pre-generated DDC data file for the engine that can be distributed to reduce runtime fetches/generation
 EnginePak=(Type=ReadPak, Filename=%ENGINEDIR%DerivedDataCache/DDC.ddp)
 
+[ZenOnlyDDC]
+MinimumDaysToKeepFile=7
+Root=(Type=KeyLength, Length=120, Inner=Zen)
+Zen=(Type=Zen, Host="http://localhost:1337" ReadOnly=false)
+
 [ZenDDC]
 MinimumDaysToKeepFile=7
 Root=(Type=KeyLength, Length=120, Inner=Hierarchy)
