Fix merge garbage - CL 16882645 never came through to Release-Engine-Staging

Edigrate [17071594]

diff --git a/Engine/Config/ConsoleVariables.ini b/Engine/Config/ConsoleVariables.ini
--- a/Engine/Config/ConsoleVariables.ini
+++ b/Engine/Config/ConsoleVariables.ini
@@ -91,11 +91,6 @@ r.DumpShaderDebugInfo=2
 ;D3D12.EmitRgpFrameMarkers=1
 ;r.AsyncPipelineCompile=0
 
-np2.bEnable=1
-np2.EnableLOD=1
-np2.MinLODForNonLocal=1 ;nonlocal net predicted physics objects will use old replication path
-np2.LODDistance=0 ;don't use distance based np lodding, avoid awkward hitches (has no effect when min lod for nonlocal is on anyway)
-
 ; 0 effectively disables physx, 1 allows physx to still run alongside ChaosPhysics
 p.chaos.AllowCreatePhysxBodies=1
 
