DDC: Removed use of the Boot cache from the backend graphs

Due to the prevalence of SSDs, the Boot cache no longer offers the performance benefits that it once did, and its startup cost now often cancels out its performance gains.

Additionally, there are cases where different values are being written for the same key, which is incorrect but has not been fixed; and the Boot cache ignores the second and subsequent writes, which can cause out-of-date content to be loaded from the Boot cache in a future session.

#rnx

diff --git a/Engine/Config/BaseEngine.ini b/Engine/Config/BaseEngine.ini
--- a/Engine/Config/BaseEngine.ini
+++ b/Engine/Config/BaseEngine.ini
@@ -1975,9 +1975,7 @@ Root=(Type=KeyLength, Length=120, Inner=AsyncPut)
 ; Configure the AsyncPut entry. It uses an AsyncPut node that then passes requests to Hierarchy
 AsyncPut=(Type=AsyncPut, Inner=Hierarchy)
 ; Configure the hierarchy entry. This uses multiple nodes that are used in order until a read is found (writes go to all writable entries)
-Hierarchy=(Type=Hierarchical, Inner=Boot, Inner=Pak, Inner=EnginePak, Inner=Local, Inner=Shared)
-; Configure the Boot node. This holds 512MB of data used to speed up booting.
-Boot=(Type=Boot, Filename="%GAMEDIR%DerivedDataCache/Boot.ddc", MaxCacheSize=512)
+Hierarchy=(Type=Hierarchical, Inner=Pak, Inner=EnginePak, Inner=Local, Inner=Shared)
 ; Configure the local node. This is a filesystem node with the following set -
 ;	Readonly: Can data be written to this layer
 ;	Clean: Perform a cleanup of old files on launch
@@ -2021,9 +2019,8 @@ Shared=(Type=FileSystem, ReadOnly=false, Clean=false, Flush=false, DeleteUnused=
 MinimumDaysToKeepFile=7
 Root=(Type=KeyLength, Length=120, Inner=AsyncPut)
 AsyncPut=(Type=AsyncPut, Inner=Hierarchy)
-Hierarchy=(Type=Hierarchical, Inner=Boot, Inner=Pak, Inner=EnginePak, Inner=Local, Inner=Seattle)
-Boot=(Type=Boot, Filename="%GAMEDIR%DerivedDataCache/Boot.ddc", MaxCacheSize=512)
-Local=(Type=FileSystem, ReadOnly=false, Clean=false, Flush=false, PurgeTransient=true, DeleteUnused=true, UnusedFileAge=34, FoldersToClean=-1, Path=%ENGINEDIR%DerivedDataCache)
+Hierarchy=(Type=Hierarchical, Inner=Pak, Inner=EnginePak, Inner=Local, Inner=Seattle)
+Local=(Type=FileSystem, ReadOnly=false, Clean=false, Flush=false, PurgeTransient=true, DeleteUnused=true, UnusedFileAge=34, FoldersToClean=-1, Path=%ENGINEDIR%DerivedDataCache, EnvPathOverride=UE-LocalDataCachePath, EditorOverrideSetting=LocalDerivedDataCache)
 Seattle=(Type=FileSystem, ReadOnly=false, Clean=false, Flush=false, DeleteUnused=true, UnusedFileAge=23, FoldersToClean=10, MaxFileChecksPerSec=1, Path=?EpicSeaDDC, EnvPathOverride=UE-SharedDataCachePath_Seattle)
 Pak=(Type=ReadPak, Filename="%GAMEDIR%DerivedDataCache/DDC.ddp")
 EnginePak=(Type=ReadPak, Filename=%ENGINEDIR%DerivedDataCache/DDC.ddp)
@@ -2032,10 +2029,9 @@ EnginePak=(Type=ReadPak, Filename=%ENGINEDIR%DerivedDataCache/DDC.ddp)
 MinimumDaysToKeepFile=7
 Root=(Type=KeyLength, Length=120, Inner=AsyncPut)
 AsyncPut=(Type=AsyncPut, Inner=Hierarchy)
-Hierarchy=(Type=Hierarchical, Inner=Boot, Inner=Pak, Inner=CompressedPak, Inner=EnginePak, Inner=EnterprisePak, Inner=Local, Inner=Shared)
-Boot=(Type=Boot, Filename="%ENGINEUSERDIR%DerivedDataCache/Boot.ddc", MaxCacheSize=512)
-Local=(Type=FileSystem, ReadOnly=false, Clean=false, Flush=false, PurgeTransient=true, DeleteUnused=true, UnusedFileAge=34, FoldersToClean=-1, Path="%ENGINEVERSIONAGNOSTICUSERDIR%DerivedDataCache", EditorOverrideSetting=LocalDerivedDataCache)
-Shared=(Type=FileSystem, ReadOnly=false, Clean=false, Flush=false, DeleteUnused=true, UnusedFileAge=10, FoldersToClean=10, MaxFileChecksPerSec=1, Path=?EpicDDC, EnvPathOverride=UE-SharedDataCachePath, EditorOverrideSetting=SharedDerivedDataCache)
+Hierarchy=(Type=Hierarchical, Inner=Pak, Inner=CompressedPak, Inner=EnginePak, Inner=EnterprisePak, Inner=Local, Inner=Shared)
+Local=(Type=FileSystem, ReadOnly=false, Clean=false, Flush=false, PurgeTransient=true, DeleteUnused=true, UnusedFileAge=34, FoldersToClean=-1, Path="%ENGINEVERSIONAGNOSTICUSERDIR%DerivedDataCache", EnvPathOverride=UE-LocalDataCachePath, EditorOverrideSetting=LocalDerivedDataCache)
+Shared=(Type=FileSystem, ReadOnly=false, Clean=false, Flush=false, DeleteUnused=true, UnusedFileAge=10, FoldersToClean=10, MaxFileChecksPerSec=1, ConsiderSlowAt=70, PromptIfMissing=false, Path=?EpicDDC, EnvPathOverride=UE-SharedDataCachePath, EditorOverrideSetting=SharedDerivedDataCache, CommandLineOverride=SharedDataCachePath)
 Pak=(Type=ReadPak, Filename="%GAMEDIR%DerivedDataCache/DDC.ddp")
 CompressedPak=(Type=ReadPak, Filename="%GAMEDIR%DerivedDataCache/Compressed.ddp", Compressed=true)
 EnginePak=(Type=ReadPak, Filename=../../../Engine/DerivedDataCache/Compressed.ddp, Compressed=true)
@@ -2045,8 +2041,7 @@ EnterprisePak=(Type=ReadPak, Filename=../../../Enterprise/DerivedDataCache/Compr
 MinimumDaysToKeepFile=7
 Root=(Type=KeyLength, Length=120, Inner=AsyncPut)
 AsyncPut=(Type=AsyncPut, Inner=Hierarchy)
-Hierarchy=(Type=Hierarchical, Inner=Boot, Inner=Pak, Inner=Local)
-Boot=(Type=Boot, Filename="%GAMEDIR%DerivedDataCache/Boot.ddc", MaxCacheSize=512)
+Hierarchy=(Type=Hierarchical, Inner=Pak, Inner=Local)
 Local=(Type=FileSystem, ReadOnly=false, Clean=false, Flush=false, PurgeTransient=true, DeleteUnused=true, UnusedFileAge=34, FoldersToClean=-1, Path=%ENGINEDIR%DerivedDataCache, EnvPathOverride=UE-LocalDataCachePath, EditorOverrideSetting=LocalDerivedDataCache)
 Pak=(Type=ReadPak, Filename="%GAMEDIR%DerivedDataCache/DDC.ddp")
 
@@ -2055,8 +2050,7 @@ Pak=(Type=ReadPak, Filename="%GAMEDIR%DerivedDataCache/DDC.ddp")
 MinimumDaysToKeepFile=0
 Root=(Type=KeyLength, Length=120, Inner=AsyncPut)
 AsyncPut=(Type=AsyncPut, Inner=Hierarchy)
-Hierarchy=(Type=Hierarchical, Inner=Boot, Inner=Local)
-Boot=(Type=Boot, Filename="%GAMEDIR%DerivedDataCache/Boot.ddc", MaxCacheSize=512)
+Hierarchy=(Type=Hierarchical, Inner=Local)
 Local=(Type=FileSystem, ReadOnly=false, Clean=false, Flush=true, PurgeTransient=true, DeleteUnused=true, UnusedFileAge=1, FoldersToClean=-1, Path="%GAMEDIR%ColdDerivedDataCache\")
 
 ; Inherits the same standalone cache from [Cold], for incremental performance measurement purposes
@@ -2064,17 +2058,15 @@ Local=(Type=FileSystem, ReadOnly=false, Clean=false, Flush=true, PurgeTransient=
 MinimumDaysToKeepFile=0
 Root=(Type=KeyLength, Length=120, Inner=AsyncPut)
 AsyncPut=(Type=AsyncPut, Inner=Hierarchy)
-Hierarchy=(Type=Hierarchical, Inner=Boot, Inner=Local)
-Boot=(Type=Boot, Filename="%GAMEDIR%DerivedDataCache/Boot.ddc", MaxCacheSize=512)
+Hierarchy=(Type=Hierarchical, Inner=Local)
 Local=(Type=FileSystem, ReadOnly=false, Clean=false, Flush=false, PurgeTransient=true, DeleteUnused=true, UnusedFileAge=1, FoldersToClean=-1, Path="%GAMEDIR%ColdDerivedDataCache\")
 
 [CreatePak]
 MinimumDaysToKeepFile=7
 Root=(Type=KeyLength, Length=120, Inner=AsyncPut)
 AsyncPut=(Type=AsyncPut, Inner=Hierarchy)
-Hierarchy=(Type=Hierarchical, Inner=Boot, Inner=PakWrite, Inner=PakRead, Inner=Local, Inner=Shared)
-Boot=(Type=Boot, Filename="%GAMEDIR%DerivedDataCache/Boot.ddc", MaxCacheSize=512)
-Local=(Type=FileSystem, ReadOnly=false, Clean=false, Flush=false, PurgeTransient=true, DeleteUnused=true, UnusedFileAge=34, FoldersToClean=-1, Path=%ENGINEDIR%DerivedDataCache)
+Hierarchy=(Type=Hierarchical, Inner=PakWrite, Inner=PakRead, Inner=Local, Inner=Shared)
+Local=(Type=FileSystem, ReadOnly=false, Clean=false, Flush=false, PurgeTransient=true, DeleteUnused=true, UnusedFileAge=34, FoldersToClean=-1, Path=%ENGINEDIR%DerivedDataCache, EnvPathOverride=UE-LocalDataCachePath, EditorOverrideSetting=LocalDerivedDataCache)
 Shared=(Type=FileSystem, ReadOnly=false, Clean=false, Flush=false, DeleteUnused=false, UnusedFileAge=23, FoldersToClean=-1, Path=?EpicDDC, EnvPathOverride=UE-SharedDataCachePath)
 AltShared=(Type=FileSystem, ReadOnly=true, Clean=false, Flush=false, DeleteUnused=false, UnusedFileAge=23, FoldersToClean=-1, Path=?EpicDDC2, EnvPathOverride=UE-SharedDataCachePath2)
 PakRead=(Type=ReadPak, Filename="%GAMEDIR%DerivedDataCache/DDC.ddp")
@@ -2084,9 +2076,8 @@ PakWrite=(Type=WritePak, Filename="%GAMEDIR%DerivedDataCache/DDC.ddp")
 MinimumDaysToKeepFile=7
 Root=(Type=KeyLength, Length=120, Inner=AsyncPut)
 AsyncPut=(Type=AsyncPut, Inner=Hierarchy)
-Hierarchy=(Type=Hierarchical, Inner=Boot, Inner=EnginePak, Inner=PakWrite, Inner=PakRead, Inner=Local, Inner=Shared)
-Boot=(Type=Boot, Filename="%GAMEDIR%DerivedDataCache/Boot.ddc", MaxCacheSize=512)
-Local=(Type=FileSystem, ReadOnly=false, Clean=false, Flush=false, PurgeTransient=true, DeleteUnused=true, UnusedFileAge=34, FoldersToClean=-1, Path=%ENGINEDIR%DerivedDataCache, EnvPathOverride=UE-LocalDataCachePath)
+Hierarchy=(Type=Hierarchical, Inner=EnginePak, Inner=PakWrite, Inner=PakRead, Inner=Local, Inner=Shared)
+Local=(Type=FileSystem, ReadOnly=false, Clean=false, Flush=false, PurgeTransient=true, DeleteUnused=true, UnusedFileAge=34, FoldersToClean=-1, Path=%ENGINEDIR%DerivedDataCache, EnvPathOverride=UE-LocalDataCachePath, EditorOverrideSetting=LocalDerivedDataCache)
 Shared=(Type=FileSystem, ReadOnly=false, Clean=false, Flush=false, DeleteUnused=false, UnusedFileAge=23, FoldersToClean=-1, Path=?EpicDDC, EnvPathOverride=UE-SharedDataCachePath, CommandLineOverride=SharedDataCachePath)
 AltShared=(Type=FileSystem, ReadOnly=true, Clean=false, Flush=false, DeleteUnused=false, UnusedFileAge=23, FoldersToClean=-1, Path=?EpicDDC2, EnvPathOverride=UE-SharedDataCachePath2)
 EnginePak=(Type=ReadPak, Filename=../../../Engine/DerivedDataCache/Compressed.ddp, Compressed=true)
@@ -2097,9 +2088,8 @@ PakWrite=(Type=WritePak, Filename="%GAMEDIR%DerivedDataCache/Compressed.ddp", Co
 MinimumDaysToKeepFile=7
 Root=(Type=KeyLength, Length=120, Inner=AsyncPut)
 AsyncPut=(Type=AsyncPut, Inner=Hierarchy)
-Hierarchy=(Type=Hierarchical, Inner=Boot, Inner=PakWrite, Inner=Local, Inner=Shared)
-Boot=(Type=Boot, Filename="%GAMEDIR%DerivedDataCache/Boot.ddc", MaxCacheSize=512)
-Local=(Type=FileSystem, ReadOnly=false, Clean=false, Flush=false, PurgeTransient=true, DeleteUnused=true, UnusedFileAge=34, FoldersToClean=-1, Path=%ENGINEDIR%DerivedDataCache, EnvPathOverride=UE-LocalDataCachePath)
+Hierarchy=(Type=Hierarchical, Inner=PakWrite, Inner=Local, Inner=Shared)
+Local=(Type=FileSystem, ReadOnly=false, Clean=false, Flush=false, PurgeTransient=true, DeleteUnused=true, UnusedFileAge=34, FoldersToClean=-1, Path=%ENGINEDIR%DerivedDataCache, EnvPathOverride=UE-LocalDataCachePath, EditorOverrideSetting=LocalDerivedDataCache)
 Shared=(Type=FileSystem, ReadOnly=false, Clean=false, Flush=false, DeleteUnused=false, UnusedFileAge=23, FoldersToClean=-1, Path=?EpicDDC, EnvPathOverride=UE-SharedDataCachePath)
 AltShared=(Type=FileSystem, ReadOnly=true, Clean=false, Flush=false, DeleteUnused=false, UnusedFileAge=23, FoldersToClean=-1, Path=?EpicDDC2, EnvPathOverride=UE-SharedDataCachePath2)
 PakWrite=(Type=WritePak, Filename=%ENGINEDIR%DerivedDataCache/Compressed.ddp, Compressed=true)
@@ -2108,9 +2098,8 @@ PakWrite=(Type=WritePak, Filename=%ENGINEDIR%DerivedDataCache/Compressed.ddp, Co
 MinimumDaysToKeepFile=7
 Root=(Type=KeyLength, Length=120, Inner=AsyncPut)
 AsyncPut=(Type=AsyncPut, Inner=Hierarchy)
-Hierarchy=(Type=Hierarchical, Inner=Boot, Inner=PakWrite, Inner=Local, Inner=Shared)
-Boot=(Type=Boot, Filename="%GAMEDIR%DerivedDataCache/Boot.ddc", MaxCacheSize=512)
-Local=(Type=FileSystem, ReadOnly=false, Clean=false, Flush=false, PurgeTransient=true, DeleteUnused=true, UnusedFileAge=34, FoldersToClean=-1, Path=%ENGINEDIR%DerivedDataCache, EnvPathOverride=UE-LocalDataCachePath)
+Hierarchy=(Type=Hierarchical, Inner=PakWrite, Inner=Local, Inner=Shared)
+Local=(Type=FileSystem, ReadOnly=false, Clean=false, Flush=false, PurgeTransient=true, DeleteUnused=true, UnusedFileAge=34, FoldersToClean=-1, Path=%ENGINEDIR%DerivedDataCache, EnvPathOverride=UE-LocalDataCachePath, EditorOverrideSetting=LocalDerivedDataCache)
 Shared=(Type=FileSystem, ReadOnly=false, Clean=false, Flush=false, DeleteUnused=false, UnusedFileAge=23, FoldersToClean=-1, Path=?EpicDDC, EnvPathOverride=UE-SharedDataCachePath)
 AltShared=(Type=FileSystem, ReadOnly=true, Clean=false, Flush=false, DeleteUnused=false, UnusedFileAge=23, FoldersToClean=-1, Path=?EpicDDC2, EnvPathOverride=UE-SharedDataCachePath2)
 PakWrite=(Type=WritePak, Filename=../../../Enterprise/DerivedDataCache/Compressed.ddp, Compressed=true)
@@ -2120,7 +2109,7 @@ MinimumDaysToKeepFile=7
 Root=(Type=KeyLength, Length=120, Inner=AsyncPut)
 AsyncPut=(Type=AsyncPut, Inner=Hierarchy)
 Hierarchy=(Type=Hierarchical, Inner=Local, Inner=Project, Inner=Shared)
-Local=(Type=FileSystem, ReadOnly=false, Clean=false, Flush=false, PurgeTransient=true, DeleteUnused=true, UnusedFileAge=34, FoldersToClean=-1, Path=%ENGINEDIR%DerivedDataCache, EnvPathOverride=UE-LocalDataCachePath)
+Local=(Type=FileSystem, ReadOnly=false, Clean=false, Flush=false, PurgeTransient=true, DeleteUnused=true, UnusedFileAge=34, FoldersToClean=-1, Path=%ENGINEDIR%DerivedDataCache, EnvPathOverride=UE-LocalDataCachePath, EditorOverrideSetting=LocalDerivedDataCache)
 Project=(Type=FileSystem, ReadOnly=false, Clean=false, Flush=true, PurgeTransient=true, DeleteUnused=true, UnusedFileAge=34, FoldersToClean=-1, Path=%GAMEDIR%ProjectDerivedData)
 Shared=(Type=FileSystem, ReadOnly=false, Clean=false, Flush=false, DeleteUnused=false, UnusedFileAge=23, FoldersToClean=-1, Path=?EpicDDC, EnvPathOverride=UE-SharedDataCachePath)
 AltShared=(Type=FileSystem, ReadOnly=true, Clean=false, Flush=false, DeleteUnused=false, UnusedFileAge=23, FoldersToClean=-1, Path=?EpicDDC2, EnvPathOverride=UE-SharedDataCachePath2)
diff --git a/Engine/Programs/ShaderCompileWorker/Config/DefaultEngine.ini b/Engine/Programs/ShaderCompileWorker/Config/DefaultEngine.ini
--- a/Engine/Programs/ShaderCompileWorker/Config/DefaultEngine.ini
+++ b/Engine/Programs/ShaderCompileWorker/Config/DefaultEngine.ini
@@ -1,4 +1,3 @@
 [DerivedDataBackendGraph]
-Boot=(Type=Boot, Filename=../../../Engine/Programs/%GAME%/DDC/Boot.ddc, MaxCacheSize=256)
 Local=(Type=FileSystem, ReadOnly=false, Clean=false, Flush=false, Path=../../../Engine/Programs/%GAME%/DDC)
 Pak=(Type=ReadPak, Filename=../../../Engine/Programs/%GAME%/DDC/DDC.ddp)
