set BaseEngine default Oodle effort to Normal
was High
CompressEffortLevel_NoEditor=Normal
so that Editor and NoEditor make the same texture by default
makes the cooked texture quality slightly lower, but faster

diff --git a/Engine/Config/BaseEngine.ini b/Engine/Config/BaseEngine.ini
--- a/Engine/Config/BaseEngine.ini
+++ b/Engine/Config/BaseEngine.ini
@@ -2835,7 +2835,7 @@ bForceRDOOff_Editor=True
 bForceRDOOff_NoEditor=True
 ; CompressEffortLevel_ is the amount of time to spend for higher quality encodings; Default, Low, Normal, High.
 CompressEffortLevel_Editor=Normal
-CompressEffortLevel_NoEditor=High
+CompressEffortLevel_NoEditor=Normal
 ; bDebugColor shows the BC1-7 format used
 bDebugColor=False
 ; bDebugDump writes out all images cooked
@@ -2849,4 +2849,4 @@ DefaultRDOLambda=30
 [OodleDataCompressionFormat]
 ; Specifies how many buffers Oodle should keep around for 
 ; reuse as opposed to heap allocating every decode.
-PreallocatedBufferCount=2
\ No newline at end of file
+PreallocatedBufferCount=2
