EditorDomain: blacklist some files currently causing a problem when resaved and loaded.
#rnx

diff --git a/Engine/Config/BaseEditor.ini b/Engine/Config/BaseEditor.ini
--- a/Engine/Config/BaseEditor.ini
+++ b/Engine/Config/BaseEditor.ini
@@ -373,8 +373,12 @@ bUseFullGCAssetClassNames=true
 EditorDomainEnabled=false
 SaveUnversioned=false
 ; +ClassBlacklist=/Long/PackageName.ClassName
+ClassBlacklist=/Script/ControlRig.ControlRigBlueprintGeneratedClass
 ; +PackageBlacklist=/Long/PackageName OR ../../../Standard/Filename/NoExtension OR D:\Absolute\Filename\NoExtension
 
+; Know bad files that are problematic if all of them come from EditorDomain. Blacklisting any one of these is sufficient to prevent the bug
++PackageBlacklist=/Niagara/Functions/RandomRangeFloat
++PackageBlacklist=/Niagara/Modules/Update/Color/ScaleColor
 
 [CookPlatformDataCacheSettings]
 MaterialInstance=100
