Fixed cvar name

diff --git a/Engine/Config/BaseScalability.ini b/Engine/Config/BaseScalability.ini
--- a/Engine/Config/BaseScalability.ini
+++ b/Engine/Config/BaseScalability.ini
@@ -494,7 +494,7 @@ fx.Niagara.QualityLevel=2
 r.Lumen.DiffuseIndirect.Allow=1
 r.Lumen.Reflections.Allow=1
 r.Lumen.TranslucencyVolume.GridPixelSize=64
-r.Lumen.Radiosity=0
+r.LumenScene.Radiosity=0
 r.Lumen.TraceMeshSDFs.Allow=0
 r.Lumen.ScreenProbeGather.RadianceCache.ProbeResolution=16
 r.Lumen.ScreenProbeGather.ScreenSpaceBentNormal=0
@@ -535,7 +535,7 @@ fx.Niagara.QualityLevel=3
 r.Lumen.DiffuseIndirect.Allow=1
 r.Lumen.Reflections.Allow=1
 r.Lumen.TranslucencyVolume.GridPixelSize=32
-r.Lumen.Radiosity=1
+r.LumenScene.Radiosity=1
 r.Lumen.TraceMeshSDFs.Allow=1
 r.Lumen.ScreenProbeGather.RadianceCache.ProbeResolution=32
 r.Lumen.ScreenProbeGather.ScreenSpaceBentNormal=1
@@ -577,7 +577,7 @@ fx.Niagara.QualityLevel=4
 r.Lumen.DiffuseIndirect.Allow=1
 r.Lumen.Reflections.Allow=1
 r.Lumen.TranslucencyVolume.GridPixelSize=32
-r.Lumen.Radiosity=1
+r.LumenScene.Radiosity=1
 r.Lumen.TraceMeshSDFs.Allow=1
 r.Lumen.ScreenProbeGather.RadianceCache.ProbeResolution=32
 r.Lumen.ScreenProbeGather.ScreenSpaceBentNormal=1
