Fix redirector & disable old ComponentParam type as Blueprintable

diff --git a/Engine/Config/BaseEngine.ini b/Engine/Config/BaseEngine.ini
--- a/Engine/Config/BaseEngine.ini
+++ b/Engine/Config/BaseEngine.ini
@@ -925,7 +925,7 @@ NearClipPlane=10.0
 +ClassRedirects=(OldName="/Script/MovieSceneTracks.MovieSceneVectorTrackExtensions",NewName="/Script/SequencerScripting.MovieSceneFloatVectorTrackExtensions")
 +StructRedirects=(OldName="/Script/MovieSceneTracks.MovieSceneVectorKeyStructBase",NewName="/Script/MovieSceneTracks.MovieSceneFloatVectorKeyStructBase")
 +StructRedirects=(OldName="/Script/MovieSceneTracks.MovieSceneVectorKeyStruct",NewName="/Script/MovieSceneTracks.MovieSceneVector3fKeyStruct")
-+StructRedirects=(OldName="AudioComponentParam", NewName="AudioParameter")
++StructRedirects=(OldName="AudioComponentParam", NewName="/Script/AudioExtensions.AudioParameter")
 +PropertyRedirects=(OldName="AudioParameter.SoundWaveParam", NewName="AudioParameter.ObjectParam")
 +FunctionRedirects=(OldName="AudioComponent.SetBoolParameter", NewName="SoundGeneratorParameterInterface.SetBoolParameter")
 +FunctionRedirects=(OldName="AudioComponent.SetFloatParameter", NewName="SoundGeneratorParameterInterface.SetFloatParameter")
diff --git a/Engine/Source/Runtime/Engine/Classes/Components/AudioComponent.h b/Engine/Source/Runtime/Engine/Classes/Components/AudioComponent.h
--- a/Engine/Source/Runtime/Engine/Classes/Components/AudioComponent.h
+++ b/Engine/Source/Runtime/Engine/Classes/Components/AudioComponent.h
@@ -126,7 +126,7 @@ enum class EAudioFaderCurve : uint8
 /**
  * Legacy struct used for storing named parameter for a given AudioComponent.
  */
-USTRUCT(BlueprintType)
+USTRUCT()
 struct UE_DEPRECATED(5.0, "FAudioComponentParam has been deprecated, use FAudioParameter") FAudioComponentParam : public FAudioParameter
 {
 	GENERATED_BODY()
