Toggle Zen AutoLaunch to be hidden by default.

diff --git a/Engine/Config/BaseEngine.ini b/Engine/Config/BaseEngine.ini
--- a/Engine/Config/BaseEngine.ini
+++ b/Engine/Config/BaseEngine.ini
@@ -2002,7 +2002,7 @@ PathToRAC=RuntimeAssetCache
 AutoLaunch=true
 
 [Zen.AutoLaunch]
-Hidden=false
+Hidden=true
 DesiredPort=1337
 WorkspaceDataPath=%ENGINEVERSIONAGNOSTICUSERDIR%Zen\Workspace
 DataPath=%ENGINEVERSIONAGNOSTICUSERDIR%Zen
