Making unititialized struct prop messages errors by default

diff --git a/Engine/Config/BaseEngine.ini b/Engine/Config/BaseEngine.ini
--- a/Engine/Config/BaseEngine.ini
+++ b/Engine/Config/BaseEngine.ini
@@ -76,9 +76,9 @@ SmallerMemoryBucket_MinGB=6
 SmallestMemoryBucket_MinGB=0
 
 [CoreUObject.UninitializedScriptStructMembersCheck]
-EngineModuleReflectedUninitializedPropertyVerbosity=Display
-ProjectModuleReflectedUninitializedPropertyVerbosity=Display
-ObjectReferenceReflectedUninitializedPropertyVerbosity=Warning
+EngineModuleReflectedUninitializedPropertyVerbosity=Error
+ProjectModuleReflectedUninitializedPropertyVerbosity=Error
+ObjectReferenceReflectedUninitializedPropertyVerbosity=Error
 
 [/Script/Engine.Engine]
 ConsoleClassName=/Script/Engine.Console
