Allow Vulkan on devices with Adreno5xx GPU running Android 10

diff --git a/Engine/Config/BaseDeviceProfiles.ini b/Engine/Config/BaseDeviceProfiles.ini
--- a/Engine/Config/BaseDeviceProfiles.ini
+++ b/Engine/Config/BaseDeviceProfiles.ini
@@ -585,7 +585,9 @@ BaseProfileName=IPadPro
 ; These rules are checked in turn and stop at the profile that matches all criteria
 [/Script/AndroidDeviceProfileSelector.AndroidDeviceProfileMatchingRules]
 MatchProfile=(Profile="Android_Adreno4xx",Match=((SourceType=SRC_GpuFamily,CompareType=CMP_Regex,MatchString="Adreno \\(TM\\) 4[0-9][0-9]")))
++MatchProfile=(Profile="Android_Adreno5xx_Low_Vulkan",Match=((SourceType=SRC_GpuFamily,CompareType=CMP_Regex,MatchString="Adreno \\(TM\\) 5[0-1][0-9]"),(SourceType=SRC_AndroidVersion, CompareType=CMP_Regex,MatchString="([0-9]+).*"),(SourceType=SRC_PreviousRegexMatch,CompareType=CMP_GreaterEqual,MatchString="10"),(SourceType=SRC_VulkanAvailable,CompareType=CMP_Equal,MatchString="true")))
 +MatchProfile=(Profile="Android_Adreno5xx_Low",Match=((SourceType=SRC_GpuFamily,CompareType=CMP_Regex,MatchString="Adreno \\(TM\\) 5[0-1][0-9]")))
++MatchProfile=(Profile="Android_Adreno5xx_Vulkan",Match=((SourceType=SRC_GpuFamily,CompareType=CMP_Regex,MatchString="Adreno \\(TM\\) 5[0-9][0-9]"),(SourceType=SRC_AndroidVersion, CompareType=CMP_Regex,MatchString="([0-9]+).*"),(SourceType=SRC_PreviousRegexMatch,CompareType=CMP_GreaterEqual,MatchString="10"),(SourceType=SRC_VulkanAvailable,CompareType=CMP_Equal,MatchString="true")))
 +MatchProfile=(Profile="Android_Adreno5xx",Match=((SourceType=SRC_GpuFamily,CompareType=CMP_Regex,MatchString="Adreno \\(TM\\) 5[0-9][0-9]")))
 ; Example profiles for Vulkan SM5
 ;+MatchProfile=(Profile="Android_Vulkan_SM5",Match=((SourceType=SRC_GpuFamily,CompareType=CMP_Regex,MatchString="Adreno \\(TM\\) 6[0-9][0-9]"),(SourceType=SRC_AndroidVersion, CompareType=CMP_Regex,MatchString="([0-9]+).*"),(SourceType=SRC_PreviousRegexMatch,CompareType=CMP_GreaterEqual,MatchString="9"),(SourceType=SRC_VulkanAvailable,CompareType=CMP_Equal,MatchString="true")))
@@ -720,10 +722,22 @@ BaseProfileName=Android_Low
 DeviceType=Android
 BaseProfileName=Android_Low
 
+[Android_Adreno5xx_Low_Vulkan DeviceProfile]
+DeviceType=Android
+BaseProfileName=Android_Adreno5xx_Low
+; enable Vulkan
++CVars=r.Android.DisableVulkanSupport=0
+
 [Android_Adreno5xx DeviceProfile]
 DeviceType=Android
 BaseProfileName=Android_Mid
 
+[Android_Adreno5xx_Vulkan DeviceProfile]
+DeviceType=Android
+BaseProfileName=Android_Adreno5xx
+; enable Vulkan
++CVars=r.Android.DisableVulkanSupport=0
+
 [Android_Adreno6xx DeviceProfile]
 DeviceType=Android
 BaseProfileName=Android_High
