EditorDomain CustomVersions: Add some construct class settings that are common in EngineTest and QAGame.
#rnx

diff --git a/Engine/Config/BaseEditor.ini b/Engine/Config/BaseEditor.ini
--- a/Engine/Config/BaseEditor.ini
+++ b/Engine/Config/BaseEditor.ini
@@ -386,6 +386,20 @@ SaveUnversioned=false
 +PackageBlockList=/Niagara/Functions/RandomRangeFloat
 +PackageBlockList=/Niagara/Modules/Update/Color/ScaleColor
 
+; Packages don't contain the full set of classes they will use when saved, because some classes
+; can be constructed by PostLoad functions or by CoreUObject savepackage functionality.
+; The EditorDomain needs to know all classes that will exist in a package when resaved so it
+; can find all CustomVersions that will exist in the package when resaved; it needs the CustomVersions
+; so that a given run of the editor can request the version the package matching its current values for
+; all CustomVersions. ClassGlobalCreateList and PostLoadCanConstructClasses record for the editor domain
+; which classes can be created when resaving a package.
+; +GlobalCanConstructClasses=<LongPackageNameOfClass>
++GlobalCanConstructClasses=/Script/CoreUObject.MetaData
++GlobalCanConstructClasses=/Script/Engine.AssetImportData
+; +PostLoadCanConstructClasses=<LongPackageNameOfExistingClass>,<LongPackageNameOfClassItCanCreateInPostLoad>
++PostLoadCanConstructClasses=/Script/Engine.StaticMesh,/Script/Engine.StaticMeshDescriptionBulkData
++PostLoadCanConstructClasses=/Script/Engine.MapBuildDataRegistry,/Script/Engine.TextureCube
+
 [TargetDomain]
 ; IterativeClassAllowListEnabled/IterativeClassAllowList/IterativeClassBlockList:
 ;    Used in HybridIterative only.
