- Restrict maximum buffer size on lower end platforms and pop a warning if we approach that size

diff --git a/Engine/Config/BaseScalability.ini b/Engine/Config/BaseScalability.ini
--- a/Engine/Config/BaseScalability.ini
+++ b/Engine/Config/BaseScalability.ini
@@ -426,9 +426,12 @@ r.SkyAtmosphere.TransmittanceLUT.UseSmallFormat=1
 r.SkyAtmosphere.TransmittanceLUT.SampleCount=10.0
 r.SkyAtmosphere.MultiScatteringLUT.SampleCount=15.0
 r.SkyLight.RealTimeReflectionCapture=0
-fx.Niagara.QualityLevel=0
 r.Lumen.DiffuseIndirect.Allow=0
 r.Lumen.Reflections.Allow=0
+; Niagara memory related (Max GPU buffer size (4mb) / Warn once when over 3.5mb)
+fx.NiagaraGPUDataWarningSize=3670016
+fx.Niagara.Emitter.MaxGPUBufferElements=1048576
+fx.Niagara.QualityLevel=0
 
 [EffectsQuality@1]
 r.TranslucencyLightingVolumeDim=32
@@ -458,9 +461,12 @@ r.SkyAtmosphere.TransmittanceLUT.UseSmallFormat=0
 r.SkyAtmosphere.TransmittanceLUT.SampleCount=10.0
 r.SkyAtmosphere.MultiScatteringLUT.SampleCount=15.0
 r.SkyLight.RealTimeReflectionCapture=0
-fx.Niagara.QualityLevel=1
 r.Lumen.DiffuseIndirect.Allow=0
 r.Lumen.Reflections.Allow=0
+; Niagara memory related (Max GPU buffer size (4mb) / Warn once when over 3.5mb)
+fx.NiagaraGPUDataWarningSize=3670016
+fx.Niagara.Emitter.MaxGPUBufferElements=1048576
+fx.Niagara.QualityLevel=1
 
 [EffectsQuality@2]
 r.TranslucencyLightingVolumeDim=48
@@ -490,7 +496,6 @@ r.SkyAtmosphere.TransmittanceLUT.UseSmallFormat=0
 r.SkyAtmosphere.TransmittanceLUT.SampleCount=10.0
 r.SkyAtmosphere.MultiScatteringLUT.SampleCount=15.0
 r.SkyLight.RealTimeReflectionCapture=1
-fx.Niagara.QualityLevel=2
 r.Lumen.DiffuseIndirect.Allow=1
 r.Lumen.Reflections.Allow=1
 r.Lumen.TranslucencyVolume.GridPixelSize=64
@@ -502,6 +507,10 @@ r.Lumen.ScreenProbeGather.DownsampleFactor=32
 r.Lumen.Reflections.DownsampleFactor=2
 r.Lumen.TranslucencyVolume.TraceFromVolume=0
 r.Lumen.TranslucencyVolume.RadianceCache.ProbeResolution=8
+; Niagara memory related (Max GPU buffer size (4mb) / Warn once when over 3.5mb)
+fx.NiagaraGPUDataWarningSize=3670016
+fx.Niagara.Emitter.MaxGPUBufferElements=1048576
+fx.Niagara.QualityLevel=2
 
 [EffectsQuality@3]
 r.TranslucencyLightingVolumeDim=64
@@ -531,7 +540,6 @@ r.SkyAtmosphere.TransmittanceLUT.UseSmallFormat=0
 r.SkyAtmosphere.TransmittanceLUT.SampleCount=10.0
 r.SkyAtmosphere.MultiScatteringLUT.SampleCount=15.0
 r.SkyLight.RealTimeReflectionCapture=1
-fx.Niagara.QualityLevel=3
 r.Lumen.DiffuseIndirect.Allow=1
 r.Lumen.Reflections.Allow=1
 r.Lumen.TranslucencyVolume.GridPixelSize=32
@@ -544,6 +552,10 @@ r.Lumen.Reflections.DownsampleFactor=1
 r.Lumen.TranslucencyVolume.TraceFromVolume=1
 r.Lumen.TranslucencyVolume.TracingOctahedronResolution=3
 r.Lumen.TranslucencyVolume.RadianceCache.ProbeResolution=8
+; Niagara memory related
+fx.NiagaraGPUDataWarningSize=0
+fx.Niagara.Emitter.MaxGPUBufferElements=0
+fx.Niagara.QualityLevel=3
 
 [EffectsQuality@Cine]
 r.TranslucencyLightingVolumeDim=64
@@ -573,7 +585,6 @@ r.SkyAtmosphere.TransmittanceLUT.UseSmallFormat=0
 r.SkyAtmosphere.TransmittanceLUT.SampleCount=30.0
 r.SkyAtmosphere.MultiScatteringLUT.SampleCount=20.0
 r.SkyLight.RealTimeReflectionCapture=1
-fx.Niagara.QualityLevel=4
 r.Lumen.DiffuseIndirect.Allow=1
 r.Lumen.Reflections.Allow=1
 r.Lumen.TranslucencyVolume.GridPixelSize=32
@@ -586,6 +597,10 @@ r.Lumen.Reflections.DownsampleFactor=1
 r.Lumen.TranslucencyVolume.TraceFromVolume=1
 r.Lumen.TranslucencyVolume.TracingOctahedronResolution=4
 r.Lumen.TranslucencyVolume.RadianceCache.ProbeResolution=16
+; Niagara memory related
+fx.NiagaraGPUDataWarningSize=0
+fx.Niagara.Emitter.MaxGPUBufferElements=0
+fx.Niagara.QualityLevel=4
 
 ;-----------------------------------------------------------------------------------------------------------------
 
