- Move change from Default -> Base scalability

#rnx

diff --git a/Engine/Config/BaseScalability.ini b/Engine/Config/BaseScalability.ini
--- a/Engine/Config/BaseScalability.ini
+++ b/Engine/Config/BaseScalability.ini
@@ -426,12 +426,9 @@ r.SkyAtmosphere.TransmittanceLUT.UseSmallFormat=1
 r.SkyAtmosphere.TransmittanceLUT.SampleCount=10.0
 r.SkyAtmosphere.MultiScatteringLUT.SampleCount=15.0
 r.SkyLight.RealTimeReflectionCapture=0
+fx.Niagara.QualityLevel=0
 r.Lumen.DiffuseIndirect.Allow=0
 r.Lumen.Reflections.Allow=0
-; Niagara memory related (Max GPU buffer size (4mb) / Warn once when over 3.5mb)
-fx.NiagaraGPUDataWarningSize=3670016
-fx.Niagara.Emitter.MaxGPUBufferElements=1048576
-fx.Niagara.QualityLevel=0
 
 [EffectsQuality@1]
 r.TranslucencyLightingVolumeDim=32
@@ -461,12 +458,9 @@ r.SkyAtmosphere.TransmittanceLUT.UseSmallFormat=0
 r.SkyAtmosphere.TransmittanceLUT.SampleCount=10.0
 r.SkyAtmosphere.MultiScatteringLUT.SampleCount=15.0
 r.SkyLight.RealTimeReflectionCapture=0
+fx.Niagara.QualityLevel=1
 r.Lumen.DiffuseIndirect.Allow=0
 r.Lumen.Reflections.Allow=0
-; Niagara memory related (Max GPU buffer size (4mb) / Warn once when over 3.5mb)
-fx.NiagaraGPUDataWarningSize=3670016
-fx.Niagara.Emitter.MaxGPUBufferElements=1048576
-fx.Niagara.QualityLevel=1
 
 [EffectsQuality@2]
 r.TranslucencyLightingVolumeDim=48
@@ -496,6 +490,7 @@ r.SkyAtmosphere.TransmittanceLUT.UseSmallFormat=0
 r.SkyAtmosphere.TransmittanceLUT.SampleCount=10.0
 r.SkyAtmosphere.MultiScatteringLUT.SampleCount=15.0
 r.SkyLight.RealTimeReflectionCapture=1
+fx.Niagara.QualityLevel=2
 r.Lumen.DiffuseIndirect.Allow=1
 r.Lumen.Reflections.Allow=1
 r.Lumen.TranslucencyVolume.GridPixelSize=64
@@ -507,10 +502,6 @@ r.Lumen.ScreenProbeGather.DownsampleFactor=32
 r.Lumen.Reflections.DownsampleFactor=2
 r.Lumen.TranslucencyVolume.TraceFromVolume=0
 r.Lumen.TranslucencyVolume.RadianceCache.ProbeResolution=8
-; Niagara memory related (Max GPU buffer size (4mb) / Warn once when over 3.5mb)
-fx.NiagaraGPUDataWarningSize=3670016
-fx.Niagara.Emitter.MaxGPUBufferElements=1048576
-fx.Niagara.QualityLevel=2
 
 [EffectsQuality@3]
 r.TranslucencyLightingVolumeDim=64
@@ -540,6 +531,7 @@ r.SkyAtmosphere.TransmittanceLUT.UseSmallFormat=0
 r.SkyAtmosphere.TransmittanceLUT.SampleCount=10.0
 r.SkyAtmosphere.MultiScatteringLUT.SampleCount=15.0
 r.SkyLight.RealTimeReflectionCapture=1
+fx.Niagara.QualityLevel=3
 r.Lumen.DiffuseIndirect.Allow=1
 r.Lumen.Reflections.Allow=1
 r.Lumen.TranslucencyVolume.GridPixelSize=32
@@ -552,10 +544,6 @@ r.Lumen.Reflections.DownsampleFactor=1
 r.Lumen.TranslucencyVolume.TraceFromVolume=1
 r.Lumen.TranslucencyVolume.TracingOctahedronResolution=3
 r.Lumen.TranslucencyVolume.RadianceCache.ProbeResolution=8
-; Niagara memory related
-fx.NiagaraGPUDataWarningSize=0
-fx.Niagara.Emitter.MaxGPUBufferElements=0
-fx.Niagara.QualityLevel=3
 
 [EffectsQuality@Cine]
 r.TranslucencyLightingVolumeDim=64
@@ -585,6 +573,7 @@ r.SkyAtmosphere.TransmittanceLUT.UseSmallFormat=0
 r.SkyAtmosphere.TransmittanceLUT.SampleCount=30.0
 r.SkyAtmosphere.MultiScatteringLUT.SampleCount=20.0
 r.SkyLight.RealTimeReflectionCapture=1
+fx.Niagara.QualityLevel=4
 r.Lumen.DiffuseIndirect.Allow=1
 r.Lumen.Reflections.Allow=1
 r.Lumen.TranslucencyVolume.GridPixelSize=32
@@ -597,10 +586,6 @@ r.Lumen.Reflections.DownsampleFactor=1
 r.Lumen.TranslucencyVolume.TraceFromVolume=1
 r.Lumen.TranslucencyVolume.TracingOctahedronResolution=4
 r.Lumen.TranslucencyVolume.RadianceCache.ProbeResolution=16
-; Niagara memory related
-fx.NiagaraGPUDataWarningSize=0
-fx.Niagara.Emitter.MaxGPUBufferElements=0
-fx.Niagara.QualityLevel=4
 
 ;-----------------------------------------------------------------------------------------------------------------
 
diff --git a/Engine/Plugins/FX/Niagara/Source/Niagara/Private/NiagaraEmitter.cpp b/Engine/Plugins/FX/Niagara/Source/Niagara/Private/NiagaraEmitter.cpp
--- a/Engine/Plugins/FX/Niagara/Source/Niagara/Private/NiagaraEmitter.cpp
+++ b/Engine/Plugins/FX/Niagara/Source/Niagara/Private/NiagaraEmitter.cpp
@@ -66,6 +66,24 @@ static FAutoConsoleVariableRef CVarNiagaraEmitterMaxGPUBufferElements(
 	TEXT("Maximum elements per GPU buffer, for example 4k elements would restrict a float buffer to be 16k maximum per buffer.\n")
 	TEXT("Note: If you request something smaller than what will satisfy a single unit of work it will be increased to that size.\n")
 	TEXT("Default 0 which will allow the buffer to be the maximum allowed by the RHI.\n"),
+	FConsoleVariableDelegate::CreateLambda(
+		[](IConsoleVariable*)
+		{
+			FNiagaraSystemUpdateContext UpdateCtx;
+			UpdateCtx.SetDestroyOnAdd(true);
+			UpdateCtx.SetOnlyActive(true);
+
+			for (TObjectIterator<UNiagaraSystem> It; It; ++It)
+			{
+				UNiagaraSystem* System = *It;
+				if ( IsValid(System) )
+				{
+					UpdateCtx.Add(System, true);
+					System->CacheFromCompiledData();
+				}
+			}
+		}
+	),
 	ECVF_Scalability
 );
 
