EditorDomain: Turn on SaveUnversioned by default.
#rnx

diff --git a/Engine/Config/BaseEditor.ini b/Engine/Config/BaseEditor.ini
--- a/Engine/Config/BaseEditor.ini
+++ b/Engine/Config/BaseEditor.ini
@@ -378,7 +378,7 @@ bUseFullGCAssetClassNames=true
 
 [EditorDomain]
 EditorDomainEnabled=false
-SaveUnversioned=false
+SaveUnversioned=true
 ; +ClassBlockList=/Long/PackageName.ClassName
 ; +PackageBlockList=/Long/PackageName OR ../../../Standard/Filename/NoExtension OR D:\Absolute\Filename\NoExtension
 
