Fix for missing/incorrect AudioParameterInterface --> AudioParameterControllerInterface redirector
#rnx

diff --git a/Engine/Config/BaseEngine.ini b/Engine/Config/BaseEngine.ini
--- a/Engine/Config/BaseEngine.ini
+++ b/Engine/Config/BaseEngine.ini
@@ -968,14 +968,16 @@ NearClipPlane=10.0
 +StructRedirects=(OldName="AudioComponentParam", NewName="/Script/AudioExtensions.AudioParameter")
 +PropertyRedirects=(OldName="AudioParameter.SoundWaveParam", NewName="AudioParameter.ObjectParam")
 
-+FunctionRedirects=(OldName="AudioComponent.SetBoolParameter", NewName="SoundParameterControllerInterface.SetBoolParameter")
-+FunctionRedirects=(OldName="AudioComponent.SetFloatParameter", NewName="SoundParameterControllerInterface.SetFloatParameter")
-+FunctionRedirects=(OldName="AudioComponent.SetIntParameter", NewName="SoundParameterControllerInterface.SetIntParameter")
-+FunctionRedirects=(OldName="SoundGeneratorParameterInterface.SetBoolParameter", NewName="SoundParameterControllerInterface.SetBoolParameter")
-+FunctionRedirects=(OldName="SoundGeneratorParameterInterface.SetFloatParameter", NewName="SoundParameterControllerInterface.SetFloatParameter")
-+FunctionRedirects=(OldName="SoundGeneratorParameterInterface.SetIntParameter", NewName="SoundParameterControllerInterface.SetIntParameter")
++FunctionRedirects=(OldName="AudioComponent.SetBoolParameter", NewName="AudioParameterControllerInterface.SetBoolParameter")
++FunctionRedirects=(OldName="AudioComponent.SetFloatParameter", NewName="AudioParameterControllerInterface.SetFloatParameter")
++FunctionRedirects=(OldName="AudioComponent.SetIntParameter", NewName="AudioParameterControllerInterface.SetIntParameter")
 +PropertyRedirects=(OldName="AudioComponent.InstanceParameters", NewName="AudioComponent.DefaultParameters")
 
++FunctionRedirects=(OldName="SoundGeneratorParameterInterface.SetBoolParameter", NewName="AudioParameterControllerInterface.SetBoolParameter")
++FunctionRedirects=(OldName="SoundGeneratorParameterInterface.SetFloatParameter", NewName="AudioParameterControllerInterface.SetFloatParameter")
++FunctionRedirects=(OldName="SoundGeneratorParameterInterface.SetIntParameter", NewName="AudioParameterControllerInterface.SetIntParameter")
++ClassRedirects=(OldName="/Script/AudioExtensions.AudioParameterInterface", NewName="/Script/AudioExtensions.AudioParameterControllerInterface")
+
 ; LiveLink data moving to runtime core module
 +ClassRedirects=(OldName="/Script/LiveLink.LiveLinkRetargetAsset", NewName="/Script/LiveLinkAnimationCore.LiveLinkRetargetAsset")
 +ClassRedirects=(OldName="/Script/LiveLink.LiveLinkRemapAsset", NewName="/Script/LiveLinkAnimationCore.LiveLinkRemapAsset")
