Added ini entry for the default interchange scene pipeline stack

diff --git a/Engine/Config/BaseEngine.ini b/Engine/Config/BaseEngine.ini
--- a/Engine/Config/BaseEngine.ini
+++ b/Engine/Config/BaseEngine.ini
@@ -2890,8 +2890,9 @@ DefaultChaosSolverActorClass=/Script/ChaosSolverEngine.ChaosSolverActor
 UnrealBuildTool=Engine/Binaries/DotNET/UnrealBuildTool/UnrealBuildTool.exe
 
 [/Script/InterchangeEngine.InterchangeProjectSettings]
-PipelineStacks=(("Generic", (Pipelines=("/Script/InterchangePipelines.InterchangeGenericAssetsPipeline"))))
+PipelineStacks=(("Generic", (Pipelines=("/Script/InterchangePipelines.InterchangeGenericAssetsPipeline"))),("Scene", (Pipelines=("/Script/InterchangePipelines.InterchangeGenericAssetsPipeline","/Script/InterchangePipelines.InterchangeGenericLevelPipeline"))))
 DefaultPipelineStack=Generic
+DefaultScenePipelineStack=Scene
 PipelineConfigurationDialogClass=/Script/InterchangeEditorPipelines.InterchangePipelineConfigurationGeneric
 bShowPipelineStacksConfigurationDialog=True
 
