BaseEngine config enable bDoUseOodleDespiteNoPluginCompression
this enables Oodle instead of zlib for ini/res/uplugin/etc.
Very little impact on total package sizes but gets us one step closer to eliminating use of zlib

diff --git a/Engine/Config/BaseEngine.ini b/Engine/Config/BaseEngine.ini
--- a/Engine/Config/BaseEngine.ini
+++ b/Engine/Config/BaseEngine.ini
@@ -2459,7 +2459,7 @@ bTouchFallbackToMouse=true
 
 [Pak]
 ; enable Oodle for files in the 'NotUsePluginCompression' list :
-;bDoUseOodleDespiteNoPluginCompression = true
+bDoUseOodleDespiteNoPluginCompression = true
 +ExtensionsToNotUsePluginCompression=uplugin
 +ExtensionsToNotUsePluginCompression=upluginmanifest
 +ExtensionsToNotUsePluginCompression=uproject
