Update audio debug command descriptions to be more accurate

diff --git a/Engine/Config/BaseInput.ini b/Engine/Config/BaseInput.ini
--- a/Engine/Config/BaseInput.ini
+++ b/Engine/Config/BaseInput.ini
@@ -147,9 +147,9 @@ MaxScrollbackSize=1024
 +ManualAutoCompleteList=(Command="au.debug.AudioResetAllDynamicSoundVolumes",Desc="Resets all dynamic volumes to unity")
 +ManualAutoCompleteList=(Command="au.debug.AudioResetDynamicSoundVolume",Desc="Resets volume for given sound type ('Class', 'Cue' or 'Wave') with provided name to unity")
 +ManualAutoCompleteList=(Command="au.debug.AudioSetDynamicSoundVolume",Desc="Name=<name> Type=<type> Vol=<vol> Sets volume for given sound type ('Class', 'Cue' or 'Wave') with provided name")
-+ManualAutoCompleteList=(Command="au.debug.AudioSoloSoundClass",Desc="<name> Solos USoundClass")
-+ManualAutoCompleteList=(Command="au.debug.AudioSoloSoundCue",Desc="<name> Solos USoundCue")
-+ManualAutoCompleteList=(Command="au.debug.AudioSoloSoundWave",Desc="<name> Solos USoundWave")
++ManualAutoCompleteList=(Command="au.debug.AudioSoloSoundClass",Desc="<name> [nonexclusive] Solos sounds using this USoundClass. If nonexclusive, existing solos will persist")
++ManualAutoCompleteList=(Command="au.debug.AudioSoloSoundCue",Desc="<name> [nonexclusive] Solos any type of USoundBase. If nonexclusive, existing solos will persist")
++ManualAutoCompleteList=(Command="au.debug.AudioSoloSoundWave",Desc="<name> [nonexclusive] Solos USoundWave. If nonexclusive, existing solos will persist")
 +ManualAutoCompleteList=(Command="au.debug.ClearSoloAudio",Desc="Clears solo'ed object")
 +ManualAutoCompleteList=(Command="au.debug.DisableLPF",Desc="Disables low-pass filter")
 +ManualAutoCompleteList=(Command="au.debug.DisableEQFilter",Desc="Disables EQ")
