Use DDC2 texture builds by default for all projects, not only ShooterGame

diff --git a/Engine/Config/BaseEditor.ini b/Engine/Config/BaseEditor.ini
--- a/Engine/Config/BaseEditor.ini
+++ b/Engine/Config/BaseEditor.ini
@@ -412,7 +412,7 @@ Material=100
 Texture2D=50
 
 [TextureBuild]
-NewTextureBuilds=false
+NewTextureBuilds=true
 
 [/Script/Localization.LocalizationSettings]
 +EngineTargetsSettings=(Name="Engine",Guid=33482D004789784C9DA695A682ACCA1B,TargetDependencies=,AdditionalManifestDependencies=,RequiredModuleNames=,GatherFromTextFiles=(IsEnabled=True,SearchDirectories=((Path="Source/Runtime/"),(Path="Source/Developer/"),(Path="Config/")),ExcludePathWildcards=((Pattern="Source/Developer/NoRedist/CommunityPortalServices/*"),(Pattern="Source/Developer/NoRedist/UnrealEngineLauncherServices/*"),(Pattern="Source/Developer/NoRedist/BuildPatchServices/*")),FileExtensions=((Pattern="cpp"),(Pattern="h"),(Pattern="c"),(Pattern="inl"),(Pattern="mm"),(Pattern="ini"))),GatherFromPackages=(IsEnabled=True,IncludePathWildcards=((Pattern="Content/*")),ExcludePathWildcards=((Pattern="Content/Editor/*"),(Pattern="Content/Tutorial/*"),(Pattern="Content/Developers/*"),(Pattern="Content/TestPackages/*"),(Pattern="Content/QA_Assets/*"),(Pattern="Content/Maps/Automation/*"),(Pattern="Content/EngineSounds/*")),FileExtensions=((Pattern="umap"),(Pattern="uasset")),ShouldGatherFromEditorOnlyData=True),GatherFromMetaData=(IsEnabled=False,IncludePathWildcards=,ExcludePathWildcards=,KeySpecifications=,ShouldGatherFromEditorOnlyData=True),NativeCultureIndex=0,SupportedCulturesStatistics=((CultureName="en"),(CultureName="es"),(CultureName="ja"),(CultureName="ko"),(CultureName="pt-BR"),(CultureName="zh-CN")))
