Enable Editor.UseLegacyGetReferencersForDeletion by default since we've found discrepancies between the original and new algorithms, preventing users to delete actors in some situations.

diff --git a/Engine/Config/ConsoleVariables.ini b/Engine/Config/ConsoleVariables.ini
--- a/Engine/Config/ConsoleVariables.ini
+++ b/Engine/Config/ConsoleVariables.ini
@@ -128,3 +128,5 @@ fx.SkipVectorVMBackendOptimizations=1
 
 ; Uncomment to select the SDK to load CAD files: "KernelIO" or "TechSoft". Default is "KernelIO"
 ;ds.CADTranslator.CADLibrary="KernelIO"
+
+Editor.UseLegacyGetReferencersForDeletion=1
\ No newline at end of file
