Move Editor.UseLegacyGetReferencersForDeletion from ConsoleVariables.ini to BaseEngine.ini file.

diff --git a/Engine/Config/BaseEngine.ini b/Engine/Config/BaseEngine.ini
--- a/Engine/Config/BaseEngine.ini
+++ b/Engine/Config/BaseEngine.ini
@@ -1854,6 +1854,8 @@ r.VSync=0
 ; Parallel rendering has not been tested in the editor so just disable it.
 r.RHICmdBypass=0
 
+Editor.UseLegacyGetReferencersForDeletion=1
+
 [SystemSettingsSplitScreen2]
 ; System settings overrides for split screen
 ; Use medium detail mode in splitscreen, this allows LD's to mark meshes as High detail and they won't render in SS
diff --git a/Engine/Config/ConsoleVariables.ini b/Engine/Config/ConsoleVariables.ini
--- a/Engine/Config/ConsoleVariables.ini
+++ b/Engine/Config/ConsoleVariables.ini
@@ -127,6 +127,4 @@ fx.SkipVectorVMBackendOptimizations=1
 ;ds.CADTranslator.DisableCADKernelTessellation=false
 
 ; Uncomment to select the SDK to load CAD files: "KernelIO" or "TechSoft". Default is "KernelIO"
-;ds.CADTranslator.CADLibrary="KernelIO"
-
-Editor.UseLegacyGetReferencersForDeletion=1
\ No newline at end of file
+;ds.CADTranslator.CADLibrary="KernelIO"
\ No newline at end of file
