Change default Zen data directory from %APPSETTINGSDIR%Zen to %APPSETTINGSDIR%Zen/Data so that we can have other subdirectories for other purposes in the future.

diff --git a/Engine/Config/BaseEngine.ini b/Engine/Config/BaseEngine.ini
--- a/Engine/Config/BaseEngine.ini
+++ b/Engine/Config/BaseEngine.ini
@@ -2079,7 +2079,7 @@ PathToRAC=RuntimeAssetCache
 AutoLaunch=true
 
 [Zen.AutoLaunch]
-DataPath=%APPSETTINGSDIR%Zen
+DataPath=%APPSETTINGSDIR%Zen/Data
 LocalDataCachePathEnvOverride=UE-LocalDataCachePath
 LocalDataCachePathEditorOverrideSetting=LocalDerivedDataCache
 ExtraArgs=--http asio --gc-interval-seconds 21600
