IoStore: Enable ShadersInIoStore for all projects on all platforms

#rnx

diff --git a/Engine/Config/BaseEngine.ini b/Engine/Config/BaseEngine.ini
--- a/Engine/Config/BaseEngine.ini
+++ b/Engine/Config/BaseEngine.ini
@@ -1109,6 +1109,7 @@ UseSeperateBulkDataFiles=False
 GameThreadHeartBeatStartSuspended=False
 GameThreadHeartBeatHitchDuration=-1.0
 GameThreadHeartBeatStackWalk=false
+AllowShadersInIoStore=True
 IgnoreInlineBulkDataReloadEnsures=False
 UsePackageTrailer=False
 HangDuration=0.0
