Moving plugin templates to correct ini file and fixing merge errors.

diff --git a/Engine/Config/BaseEditor.ini b/Engine/Config/BaseEditor.ini
--- a/Engine/Config/BaseEditor.ini
+++ b/Engine/Config/BaseEditor.ini
@@ -541,3 +541,10 @@ bValidateOnSave=true
 [/Script/UnrealEd.DiffAssetRegistriesCommandlet]
 P4EngineBasePath="/Engine/Content/"
 P4EngineAssetPath="/Engine/"
+
+
+[EditorSettings]
+bCanModifyPluginsFromBrowser=false
+
+[/Script/GameFeaturesEditor.GameFeaturesEditorSettings]
++PluginTemplates=(Path=(Path="../../Plugins/Experimental/GameFeatures/Templates/GameFeaturePluginContentOnly"),Label=NSLOCTEXT("[/Script/GameFeaturesEditor]", "E17DB7E94A2C0F5968A52E89959864C0", "Game Feature (Content Only)"),Description=NSLOCTEXT("[/Script/GameFeaturesEditor]", "CB5136B64AFAFBA5EDEA078FCAC7BCC9", "Create a new Game Feature Plugin."),DefaultGameFeatureDataClass=None,DefaultGameFeatureDataName="")
diff --git a/Engine/Config/BaseEngine.ini b/Engine/Config/BaseEngine.ini
--- a/Engine/Config/BaseEngine.ini
+++ b/Engine/Config/BaseEngine.ini
@@ -2936,8 +2936,3 @@ EnvironmentCollisionThickness = 2.f
 ; Limit the number of extracted shapes per component as these collisions are very expensive
 EnvironmentCollisionMaxShapes = 32
 
-[EditorSettings]
-bCanModifyPluginsFromBrowser=false
-
-[/Script/GameFeaturesEditor.GameFeaturesEditorSettings]
-+PluginTemplates=(Path=(Path="../../Plugins/Experimental/GameFeatures/Templates/GameFeaturePluginContentOnly"),Label=NSLOCTEXT("[/Script/GameFeaturesEditor]", "E17DB7E94A2C0F5968A52E89959864C0", "Game Feature (Content Only)"),Description=NSLOCTEXT("[/Script/GameFeaturesEditor]", "CB5136B64AFAFBA5EDEA078FCAC7BCC9", "Create a new Game Feature Plugin."),DefaultGameFeatureDataClass=None,DefaultGameFeatureDataName="")
\ No newline at end of file
