Excluding audio classes from dedicated server cooks

diff --git a/Engine/Config/BaseEngine.ini b/Engine/Config/BaseEngine.ini
--- a/Engine/Config/BaseEngine.ini
+++ b/Engine/Config/BaseEngine.ini
@@ -2640,6 +2640,9 @@ DefaultASTCQualityBySize=3
 +ClassesExcludedOnDedicatedServer=SubUVAnimation
 +ClassesExcludedOnDedicatedServer=SoundNode
 +ClassesExcludedOnDedicatedServer=GameplayEffectUIData
++ClassesExcludedOnDedicatedServer=SoundWave
++ClassesExcludedOnDedicatedServer=SoundCue
++ClassesExcludedOnDedicatedServer=SoundBase
 +ClassesExcludedOnDedicatedClient=WidgetBlueprint
 +ClassesExcludedOnDedicatedClient=GroupActor
 +ClassesExcludedOnDedicatedClient=MetaData
