Set open world map as the default editor map for new projects without their own template map

diff --git a/Engine/Config/BaseEngine.ini b/Engine/Config/BaseEngine.ini
--- a/Engine/Config/BaseEngine.ini
+++ b/Engine/Config/BaseEngine.ini
@@ -4,8 +4,8 @@
 
 [/Script/EngineSettings.GameMapsSettings]
 GameInstanceClass=/Script/Engine.GameInstance
-EditorStartupMap=/Engine/Maps/Templates/Template_Default
-GameDefaultMap=/Engine/Maps/Entry
+EditorStartupMap=/Engine/Maps/Templates/OpenWorld
+GameDefaultMap=/Engine/Maps/Templates/OpenWorld
 ServerDefaultMap=/Engine/Maps/Entry
 GlobalDefaultGameMode="/Script/Engine.GameModeBase"
 ; Example for map prefix and game mode name setting
diff --git a/Templates/TP_Blank/Config/DefaultEngine.ini b/Templates/TP_Blank/Config/DefaultEngine.ini
--- a/Templates/TP_Blank/Config/DefaultEngine.ini
+++ b/Templates/TP_Blank/Config/DefaultEngine.ini
@@ -1,5 +1,5 @@
 
 
 [/Script/EngineSettings.GameMapsSettings]
-GameDefaultMap=/Engine/Maps/Templates/Template_Default.Template_Default
+GameDefaultMap=/Engine/Maps/Templates/OpenWorld
 
diff --git a/Templates/TP_BlankBP/Config/DefaultEngine.ini b/Templates/TP_BlankBP/Config/DefaultEngine.ini
--- a/Templates/TP_BlankBP/Config/DefaultEngine.ini
+++ b/Templates/TP_BlankBP/Config/DefaultEngine.ini
@@ -1,5 +1,5 @@
 
 
 [/Script/EngineSettings.GameMapsSettings]
-GameDefaultMap=/Engine/Maps/Templates/Template_Default.Template_Default
+GameDefaultMap=/Engine/Maps/Templates/OpenWorld
 
