Remove duplicate AllowShadersInIoStore=True entry in BaseEngine.ini

#rnx

diff --git a/Engine/Config/BaseEngine.ini b/Engine/Config/BaseEngine.ini
--- a/Engine/Config/BaseEngine.ini
+++ b/Engine/Config/BaseEngine.ini
@@ -1121,7 +1121,6 @@ AssetLogShowsAbsolutePath=True
 TaskThreadStackSize=393216
 PoolThreadStackSize=131072
 DefaultThreadStackSize=1048576
-AllowShadersInIoStore=True
 
 [Core.System.Experimental]
 EnablePackageSidecarSaving=False
