Adjust cine scability settings to minimize the chance of page pool overflow in scenes with many lights.

diff --git a/Engine/Config/BaseScalability.ini b/Engine/Config/BaseScalability.ini
--- a/Engine/Config/BaseScalability.ini
+++ b/Engine/Config/BaseScalability.ini
@@ -236,10 +236,10 @@ r.VolumetricFog.HistoryMissSupersampleCount=16
 r.LightMaxDrawDistanceScale=1
 r.CapsuleShadows=1
 r.Shadow.Virtual.MaxPhysicalPages=8192
-r.Shadow.Virtual.ResolutionLodBiasDirectional=-2.5
-r.Shadow.Virtual.ResolutionLodBiasLocal=-1.0
+r.Shadow.Virtual.ResolutionLodBiasDirectional=-1.5
+r.Shadow.Virtual.ResolutionLodBiasLocal=0.0
 r.Shadow.Virtual.SMRT.RayCountDirectional=16
-r.Shadow.Virtual.SMRT.SamplesPerRayDirectional=16
+r.Shadow.Virtual.SMRT.SamplesPerRayDirectional=8
 r.Shadow.Virtual.SMRT.RayCountLocal=16
 r.Shadow.Virtual.SMRT.SamplesPerRayLocal=8
 
