Disable gc.AssetClustering by default for all projects since too many projects are experiencing problems with materials.
This means that all three cluster options (Asset, Actor and Blueprint) are disabled by default and instead opt-in for games that find them favourable.
Keep all three cluster options enabled in ShooterGame for regression testing.

#rnx

diff --git a/Engine/Config/BaseEngine.ini b/Engine/Config/BaseEngine.ini
--- a/Engine/Config/BaseEngine.ini
+++ b/Engine/Config/BaseEngine.ini
@@ -1195,6 +1195,7 @@ gc.MaxObjectsInEditor=25165824
 gc.IncrementalBeginDestroyEnabled=True
 gc.CreateGCClusters=True
 gc.MinGCClusterSize=5
+gc.AssetClustreringEnabled=False
 gc.ActorClusteringEnabled=False
 gc.BlueprintClusteringEnabled=False
 gc.UseDisregardForGCOnDedicatedServers=False
