Enable Distance Field Ambient Occlusion on Medium scalability level, as a fallback from Lumen GI for games using fully dynamic lighting

diff --git a/Engine/Config/BaseScalability.ini b/Engine/Config/BaseScalability.ini
--- a/Engine/Config/BaseScalability.ini
+++ b/Engine/Config/BaseScalability.ini
@@ -150,7 +150,8 @@ r.Shadow.DistanceScale=0.7
 r.Shadow.CSM.TransitionScale=0.25
 r.Shadow.PreShadowResolutionFactor=0.5
 r.DistanceFieldShadowing=0
-r.DistanceFieldAO=0
+; Enable DFAO when Lumen is off
+r.DistanceFieldAO=1
 r.VolumetricFog=0
 r.LightMaxDrawDistanceScale=.5
 r.CapsuleShadows=1
