Raised the quality of Lumen Final Gather in Cinematic scalability level

diff --git a/Engine/Config/BaseScalability.ini b/Engine/Config/BaseScalability.ini
--- a/Engine/Config/BaseScalability.ini
+++ b/Engine/Config/BaseScalability.ini
@@ -260,6 +260,7 @@ r.Lumen.TraceMeshSDFs.Allow=0
 r.Lumen.ScreenProbeGather.RadianceCache.ProbeResolution=16
 r.Lumen.ScreenProbeGather.ScreenSpaceBentNormal=0
 r.Lumen.ScreenProbeGather.DownsampleFactor=32
+r.Lumen.ScreenProbeGather.TracingOctahedronResolution=8
 r.Lumen.TranslucencyVolume.TraceFromVolume=0
 r.Lumen.TranslucencyVolume.RadianceCache.ProbeResolution=8
 
@@ -271,6 +272,7 @@ r.Lumen.TraceMeshSDFs.Allow=1
 r.Lumen.ScreenProbeGather.RadianceCache.ProbeResolution=32
 r.Lumen.ScreenProbeGather.ScreenSpaceBentNormal=1
 r.Lumen.ScreenProbeGather.DownsampleFactor=16
+r.Lumen.ScreenProbeGather.TracingOctahedronResolution=8
 r.Lumen.TranslucencyVolume.TraceFromVolume=1
 r.Lumen.TranslucencyVolume.TracingOctahedronResolution=3
 r.Lumen.TranslucencyVolume.RadianceCache.ProbeResolution=8
@@ -282,7 +284,8 @@ r.LumenScene.Radiosity=1
 r.Lumen.TraceMeshSDFs.Allow=1
 r.Lumen.ScreenProbeGather.RadianceCache.ProbeResolution=32
 r.Lumen.ScreenProbeGather.ScreenSpaceBentNormal=1
-r.Lumen.ScreenProbeGather.DownsampleFactor=16
+r.Lumen.ScreenProbeGather.DownsampleFactor=8
+r.Lumen.ScreenProbeGather.TracingOctahedronResolution=16
 r.Lumen.TranslucencyVolume.TraceFromVolume=1
 r.Lumen.TranslucencyVolume.TracingOctahedronResolution=4
 r.Lumen.TranslucencyVolume.RadianceCache.ProbeResolution=16
