Only applies r.MotionBlur.HalfResGather=1 on low and medium scalability settings

diff --git a/Engine/Config/BaseScalability.ini b/Engine/Config/BaseScalability.ini
--- a/Engine/Config/BaseScalability.ini
+++ b/Engine/Config/BaseScalability.ini
@@ -317,6 +317,7 @@ r.Lumen.Reflections.DownsampleFactor=1
 
 [PostProcessQuality@0]
 r.MotionBlurQuality=0
+r.MotionBlur.HalfResGather=1
 r.AmbientOcclusionMipLevelFactor=1.0
 r.AmbientOcclusionMaxQuality=0
 r.AmbientOcclusionLevels=0
@@ -337,6 +338,7 @@ r.Tonemapper.Quality=0
 
 [PostProcessQuality@1]
 r.MotionBlurQuality=3
+r.MotionBlur.HalfResGather=1
 r.AmbientOcclusionMipLevelFactor=1.0
 r.AmbientOcclusionMaxQuality=60
 r.AmbientOcclusionLevels=-1
@@ -369,6 +371,7 @@ r.DOF.Kernel.MaxBackgroundRadius=0.006   ; required because low gathering and no
 
 [PostProcessQuality@2]
 r.MotionBlurQuality=3
+r.MotionBlur.HalfResGather=0
 r.AmbientOcclusionMipLevelFactor=0.6
 r.AmbientOcclusionMaxQuality=100
 r.AmbientOcclusionLevels=-1
@@ -403,6 +406,7 @@ r.DOF.Kernel.MaxBackgroundRadius=0.012   ; required because of AccumulatorQualit
 
 [PostProcessQuality@3]
 r.MotionBlurQuality=4
+r.MotionBlur.HalfResGather=0
 r.AmbientOcclusionMipLevelFactor=0.4
 r.AmbientOcclusionMaxQuality=100
 r.AmbientOcclusionLevels=-1
@@ -439,6 +443,7 @@ r.DOF.Kernel.MaxBackgroundRadius=0.025
 
 [PostProcessQuality@Cine]
 r.MotionBlurQuality=4
+r.MotionBlur.HalfResGather=0
 r.AmbientOcclusionMipLevelFactor=0.4
 r.AmbientOcclusionMaxQuality=100
 r.AmbientOcclusionLevels=-1
diff --git a/Engine/Source/Runtime/Renderer/Private/PostProcess/PostProcessMotionBlur.cpp b/Engine/Source/Runtime/Renderer/Private/PostProcess/PostProcessMotionBlur.cpp
--- a/Engine/Source/Runtime/Renderer/Private/PostProcess/PostProcessMotionBlur.cpp
+++ b/Engine/Source/Runtime/Renderer/Private/PostProcess/PostProcessMotionBlur.cpp
@@ -51,7 +51,7 @@ namespace
 	TAutoConsoleVariable<int32> CVarMotionBlurHalfResGather(
 		TEXT("r.MotionBlur.HalfResGather"), 1,
 		TEXT("Whether to do motion blur filter dynamically at half res under heavy motion."),
-		ECVF_RenderThreadSafe);
+		ECVF_Scalability | ECVF_RenderThreadSafe);
 
 	TAutoConsoleVariable<int32> CVarMotionBlurAllowExternalVelocityFlatten(
 		TEXT("r.MotionBlur.AllowExternalVelocityFlatten"), 1,
