Change Zen's autolaunched local cache to use a 34 day eviction policy like the default file local cache.

diff --git a/Engine/Config/BaseEngine.ini b/Engine/Config/BaseEngine.ini
--- a/Engine/Config/BaseEngine.ini
+++ b/Engine/Config/BaseEngine.ini
@@ -2092,7 +2092,7 @@ AutoLaunch=true
 DataPath=%APPSETTINGSDIR%Zen/Data
 LocalDataCachePathEnvOverride=UE-LocalDataCachePath
 LocalDataCachePathEditorOverrideSetting=LocalDerivedDataCache
-ExtraArgs=--http asio --gc-interval-seconds 21600
+ExtraArgs=--http asio --gc-cache-duration-seconds 2937600 --gc-interval-seconds 21600
 
 [Zen.ConnectExisting]
 HostName=localhost
