Tweak default SMRT scalability settings

diff --git a/Engine/Config/BaseScalability.ini b/Engine/Config/BaseScalability.ini
--- a/Engine/Config/BaseScalability.ini
+++ b/Engine/Config/BaseScalability.ini
@@ -213,9 +213,9 @@ r.Shadow.Virtual.MaxPhysicalPages=4096
 r.Shadow.Virtual.ResolutionLodBiasDirectional=-1.5
 r.Shadow.Virtual.ResolutionLodBiasLocal=0.0
 r.Shadow.Virtual.SMRT.RayCountDirectional=8
-r.Shadow.Virtual.SMRT.SamplesPerRayDirectional=8
+r.Shadow.Virtual.SMRT.SamplesPerRayDirectional=4
 r.Shadow.Virtual.SMRT.RayCountLocal=8
-r.Shadow.Virtual.SMRT.SamplesPerRayLocal=8
+r.Shadow.Virtual.SMRT.SamplesPerRayLocal=4
 
 [ShadowQuality@Cine]
 r.LightFunctionQuality=1
