Add package and class blocklist items to avoid EditorDomain wasting time resaving things it cant persist.

diff --git a/Engine/Config/BaseEditor.ini b/Engine/Config/BaseEditor.ini
--- a/Engine/Config/BaseEditor.ini
+++ b/Engine/Config/BaseEditor.ini
@@ -367,6 +367,8 @@ SaveUnversioned=true
 ; +ClassBlockList=/Long/PackageName.ClassName
 ; +PackageBlockList=/Long/PackageName OR ../../../Standard/Filename/NoExtension OR D:\Absolute\Filename\NoExtension
 
++ClassBlockList=/Script/Engine.StaticMesh
+
 ; Known bad files that are problematic if all of them come from EditorDomain. Blocking any one of these is sufficient to prevent the bug.
 +PackageBlockList=/Niagara/Functions/RandomRangeFloat
 +PackageBlockList=/Niagara/Modules/Update/Color/ScaleColor
