Change default compressionblocksize back to 64 KB from 256 KB
to prevent mismatch between FileCache tile size and compressionblocksize
projects without VT/Nanite benefit from larger compressionblocksize

diff --git a/Engine/Config/BaseGame.ini b/Engine/Config/BaseGame.ini
--- a/Engine/Config/BaseGame.ini
+++ b/Engine/Config/BaseGame.ini
@@ -88,7 +88,8 @@ UsePakFile=True
 bUseIoStore=True
 bCompressed=True
 PackageCompressionFormat=Oodle
-PackageAdditionalCompressionOptions=-compressionblocksize=256KB
+; default compression block size is 64 KB
+;PackageAdditionalCompressionOptions=-compressionblocksize=256KB
 PackageCompressionMethod=Kraken
 PackageCompressionLevel_Distribution=7
 PackageCompressionLevel_TestShipping=5
