add default values for TextureEncodingProjectSettings to BaseEngine.ini

diff --git a/Engine/Config/BaseEngine.ini b/Engine/Config/BaseEngine.ini
--- a/Engine/Config/BaseEngine.ini
+++ b/Engine/Config/BaseEngine.ini
@@ -2924,13 +2924,25 @@ bForceAllBC23ToBC7=False
 ; bDebugColor shows the BC1-7 format used
 bDebugColor=False
 ; DebugDump writes out the input and encoded images to Saved/OodleDebugImages, 
-; based on a wilcard match. DebugDumpFilter=* matches all. Case insensitive.
+; based on a wildcard match. DebugDumpFilter=* matches all. Case insensitive.
 DebugDumpFilter=
 ; LogVerbosity 0,1, or 2
 LogVerbosity=0
 ; Lambda controls the quality vs size tradeoff for RDO
 GlobalLambdaMultiplier=1.000000
 
+[/Script/Engine.TextureEncodingProjectSettings]
+bFinalUsesRDO=false
+FinalRDOLambda=30
+FinalEffortLevel=Normal
+FinalUniversalTiling=Disabled
+bFastUsesRDO=false
+FastRDOLambda=30
+FastEffortLevel=Normal
+FastUniversalTiling=Disabled
+CookUsesSpeed=Final
+EditorUsesSpeed=FinalIfAvailable
+
 [OodleDataCompressionFormat]
 ; Specifies how many buffers Oodle should keep around for 
 ; reuse as opposed to heap allocating every decode.
