OculusVR plugin defaults to Legacy Oculus SDK

diff --git a/Engine/Config/BaseEngine.ini b/Engine/Config/BaseEngine.ini
--- a/Engine/Config/BaseEngine.ini
+++ b/Engine/Config/BaseEngine.ini
@@ -2710,6 +2710,9 @@ SteamVR=30
 OpenXRHMD=20
 WindowsMixedRealityHMD=10
 
+[/Script/OculusHMD.OculusHMDRuntimeSettings]
+XrApi=NativeOpenXR
+
 [/Script/Engine.AISystemBase]
 AISystemModuleName=AIModule
 AISystemClassName=/Script/AIModule.AISystem
