Updates Medium anti-aliasing scalability with velocity hole filling and rejection filtering

diff --git a/Engine/Config/BaseScalability.ini b/Engine/Config/BaseScalability.ini
--- a/Engine/Config/BaseScalability.ini
+++ b/Engine/Config/BaseScalability.ini
@@ -53,9 +53,9 @@ r.TemporalAA.Quality=1
 r.TSR.History.R11G11B10=1
 r.TSR.History.ScreenPercentage=100
 r.TSR.History.UpdateQuality=1
-r.TSR.ShadingRejection.SpatialFilter=0
+r.TSR.ShadingRejection.SpatialFilter=1
 r.TSR.Velocity.Extrapolation=0
-r.TSR.Velocity.HoleFill=0
+r.TSR.Velocity.HoleFill=1
 r.TSR.RejectionAntiAliasingQuality=1
 
 [AntiAliasingQuality@2]
