Add 12 iPads and 4 new iPhone13s. Also deletes iOS devices that don't support iOS14+

diff --git a/Engine/Config/BaseDeviceProfiles.ini b/Engine/Config/BaseDeviceProfiles.ini
--- a/Engine/Config/BaseDeviceProfiles.ini
+++ b/Engine/Config/BaseDeviceProfiles.ini
@@ -9,19 +9,13 @@
 +DeviceProfileNameAndTypes=WindowsServer,Windows
 +DeviceProfileNameAndTypes=WindowsClient,Windows
 +DeviceProfileNameAndTypes=IOS,IOS
-+DeviceProfileNameAndTypes=iPadAir,IOS
 +DeviceProfileNameAndTypes=iPadAir2,IOS
 +DeviceProfileNameAndTypes=iPadAir3,IOS
 +DeviceProfileNameAndTypes=iPadAir4,IOS
-+DeviceProfileNameAndTypes=iPadMini2,IOS
-+DeviceProfileNameAndTypes=iPadMini3,IOS
 +DeviceProfileNameAndTypes=iPadMini4,IOS
 +DeviceProfileNameAndTypes=iPadMini5,IOS
-+DeviceProfileNameAndTypes=iPodTouch6,IOS
++DeviceProfileNameAndTypes=iPadMini6,IOS
 +DeviceProfileNameAndTypes=iPodTouch7,IOS
-+DeviceProfileNameAndTypes=iPhone5S,IOS
-+DeviceProfileNameAndTypes=iPhone6,IOS
-+DeviceProfileNameAndTypes=iPhone6Plus,IOS
 +DeviceProfileNameAndTypes=iPhone6S,IOS
 +DeviceProfileNameAndTypes=iPhone6SPlus,IOS
 +DeviceProfileNameAndTypes=iPhone7,IOS
@@ -41,6 +35,10 @@
 +DeviceProfileNameAndTypes=iPhone12,IOS
 +DeviceProfileNameAndTypes=iPhone12Pro,IOS
 +DeviceProfileNameAndTypes=iPhone12ProMax,IOS
++DeviceProfileNameAndTypes=iPhone13Mini,IOS
++DeviceProfileNameAndTypes=iPhone13,IOS
++DeviceProfileNameAndTypes=iPhone13Pro,IOS
++DeviceProfileNameAndTypes=iPhone13ProMax,IOS
 +DeviceProfileNameAndTypes=iPadPro,IOS
 +DeviceProfileNameAndTypes=iPadPro105,IOS
 +DeviceProfileNameAndTypes=iPadPro129,IOS
@@ -50,10 +48,13 @@
 +DeviceProfileNameAndTypes=iPad6,IOS
 +DeviceProfileNameAndTypes=iPad7,IOS
 +DeviceProfileNameAndTypes=iPad8,IOS
++DeviceProfileNameAndTypes=iPad9,IOS
 +DeviceProfileNameAndTypes=iPadPro11,IOS
 +DeviceProfileNameAndTypes=iPadPro2_11,IOS
++DeviceProfileNameAndTypes=iPadPro3_11,IOS
 +DeviceProfileNameAndTypes=iPadPro3_129,IOS
 +DeviceProfileNameAndTypes=iPadPro4_129,IOS
++DeviceProfileNameAndTypes=iPadPro5_129,IOS
 +DeviceProfileNameAndTypes=AppleTV,IOS
 +DeviceProfileNameAndTypes=AppleTV4K,IOS
 +DeviceProfileNameAndTypes=TVOS,AppleTV
@@ -166,9 +167,6 @@
 ; iOS Devices
 
 [IOSDeviceMappings]
-iPhone6,[1,2]=iPhone5S
-iPhone7,1=iPhone6Plus
-iPhone7,2=iPhone6
 iPhone8,1=iPhone6S
 iPhone8,2=iPhone6SPlus
 iPhone8,4=iPhoneSE
@@ -188,11 +186,11 @@
 iPhone13,2=iPhone12
 iPhone13,3=iPhone12Pro
 iPhone13,4=iPhone12ProMax
-iPod7,1=iPodTouch6
+iPhone14,2=iPhone13Pro
+iPhone14,3=iPhone13ProMax
+iPhone14,4=iPhone13Mini
+iPhone14,5=iPhone13
 iPod9,1=iPodTouch7
-iPad4,[1-3]=iPadAir
-iPad4,[4-6]=iPadMini2
-iPad4,[7-9]=iPadMini3
 iPad5,[1,2]=iPadMini4
 iPad5,[3,4]=iPadAir2
 iPad6,[3,4]=iPadPro97
@@ -213,13 +211,17 @@
 iPad11,[1,2]=iPadMini5
 iPad11,[3,4]=iPadAir3
 iPad11,[6,7]=iPad8
+iPad12,[1,2]=iPad9
 iPad13,[1,2]=iPadAir4
+iPad13,[4-7]=iPadPro3_11
+iPad13,[8-11]=iPadPro5_129
+iPad14,[1,2]=iPadMini6
 AppleTV5,=AppleTV
 AppleTV6,=AppleTV4K
 ; Below are generic fallbacks for unrecognized devices. Add any specific devices above this section.
-iPhone=iPhone12
+iPhone=iPhone13
 iPod=iPodTouch7
-iPad=iPad8
+iPad=iPad9
 AppleTV=AppleTV4K
 
 [IOS DeviceProfile]
@@ -251,134 +253,22 @@
 ; If any PreviewAllowlistCVars is set, cvars are denied by default.
 PreviewAllowlistCVars=sg.
 
-[iPad2 DeviceProfile]
-DeviceType=IOS
-BaseProfileName=IOS
-+CVars=ios.PhysicalScreenDensity=264
-+CVars=r.RenderTargetSwitchWorkaround=1
-+CVars=r.Mobile.AmbientOcclusionQuality=0
-+CVars=r.Mobile.PixelProjectedReflectionQuality=0
-
-[iPad3 DeviceProfile]
-DeviceType=IOS
-BaseProfileName=IOS
-+CVars=ios.PhysicalScreenDensity=264
-+CVars=r.RenderTargetSwitchWorkaround=1
-+CVars=r.Mobile.AmbientOcclusionQuality=0
-+CVars=r.Mobile.PixelProjectedReflectionQuality=0
-
-[iPad4 DeviceProfile]
-DeviceType=IOS
-BaseProfileName=IOS
-+CVars=ios.PhysicalScreenDensity=264
-+CVars=r.RenderTargetSwitchWorkaround=1
-+CVars=r.Mobile.AmbientOcclusionQuality=0
-+CVars=r.Mobile.PixelProjectedReflectionQuality=0
-
-[iPadAir DeviceProfile]
-DeviceType=IOS
-BaseProfileName=IOS
-+CVars=ios.PhysicalScreenDensity=264
-+CVars=sg.PostProcessQuality=1
-+CVars=r.Mobile.AmbientOcclusionQuality=0
-+CVars=r.Mobile.PixelProjectedReflectionQuality=0
-
-[iPadAir2 DeviceProfile]
-DeviceType=IOS
-BaseProfileName=IOS
-+CVars=ios.PhysicalScreenDensity=264
-+CVars=sg.PostProcessQuality=3
-+CVars=sg.ShadowQuality=2
-+CVars=r.Mobile.AmbientOcclusionQuality=0
-+CVars=r.Mobile.PixelProjectedReflectionQuality=0
-
-[iPadMini DeviceProfile]
-DeviceType=IOS
-BaseProfileName=IOS
-+CVars=ios.PhysicalScreenDensity=401
-+CVars=r.RenderTargetSwitchWorkaround=1
-+CVars=r.Mobile.AmbientOcclusionQuality=0
-+CVars=r.Mobile.PixelProjectedReflectionQuality=0
-
-[iPadMini2 DeviceProfile]
-DeviceType=IOS
-BaseProfileName=iPadAir
-+CVars=ios.PhysicalScreenDensity=401
-+CVars=r.Mobile.AmbientOcclusionQuality=0
-+CVars=r.Mobile.PixelProjectedReflectionQuality=0
-
-[iPadMini3 DeviceProfile]
-DeviceType=IOS
-BaseProfileName=iPadMini2
-+CVars=ios.PhysicalScreenDensity=401
-+CVars=r.Mobile.AmbientOcclusionQuality=0
-+CVars=r.Mobile.PixelProjectedReflectionQuality=0
-
-[iPadMini4 DeviceProfile]
-DeviceType=IOS
-BaseProfileName=iPadMini2
-+CVars=ios.PhysicalScreenDensity=401
-+CVars=sg.PostProcessQuality=2
-+CVars=r.Mobile.AmbientOcclusionQuality=0
-+CVars=r.Mobile.PixelProjectedReflectionQuality=0
-
-[iPhone5S DeviceProfile]
-DeviceType=IOS
-BaseProfileName=IOS
-+CVars=ios.PhysicalScreenDensity=326
-+CVars=r.MobileContentScaleFactor=2
-+CVars=sg.PostProcessQuality=3
-+CVars=r.Mobile.AmbientOcclusionQuality=0
-+CVars=r.Mobile.PixelProjectedReflectionQuality=0
-
-[iPodTouch6 DeviceProfile]
-DeviceType=IOS
-BaseProfileName=iPhone6
-+CVars=ios.PhysicalScreenDensity=326
-+CVars=r.MobileContentScaleFactor=2
-+CVars=r.Mobile.AmbientOcclusionQuality=0
-+CVars=r.Mobile.PixelProjectedReflectionQuality=0
-
-[iPodTouch7 DeviceProfile]
-DeviceType=IOS
-BaseProfileName=iPhone7
-+CVars=ios.PhysicalScreenDensity=326
-+CVars=r.MobileContentScaleFactor=2
-+CVars=r.Mobile.AmbientOcclusionQuality=0
-+CVars=r.Mobile.PixelProjectedReflectionQuality=0
-
-[iPhone6 DeviceProfile]
-DeviceType=IOS
-BaseProfileName=IOS
-+CVars=ios.PhysicalScreenDensity=326
-+CVars=r.MobileContentScaleFactor=2
-+CVars=sg.PostProcessQuality=3
-+CVars=r.Mobile.AmbientOcclusionQuality=0
-+CVars=r.Mobile.PixelProjectedReflectionQuality=0
-
-[iPhone6Plus DeviceProfile]
-DeviceType=IOS
-BaseProfileName=IOS
-+CVars=ios.PhysicalScreenDensity=401
-+CVars=r.MobileContentScaleFactor=2
-+CVars=sg.PostProcessQuality=3
-+CVars=r.Mobile.AmbientOcclusionQuality=0
-+CVars=r.Mobile.PixelProjectedReflectionQuality=0
-
 [iPhone6S DeviceProfile]
 DeviceType=IOS
-BaseProfileName=iPhone6
+BaseProfileName=IOS
 +CVars=ios.PhysicalScreenDensity=326
 +CVars=r.MobileContentScaleFactor=2
++CVars=sg.PostProcessQuality=3
 +CVars=sg.ShadowQuality=2
 +CVars=r.Mobile.AmbientOcclusionQuality=0
 +CVars=r.Mobile.PixelProjectedReflectionQuality=0
 
 [iPhone6SPlus DeviceProfile]
 DeviceType=IOS
-BaseProfileName=iPhone6Plus
+BaseProfileName=IOS
 +CVars=ios.PhysicalScreenDensity=401
 +CVars=r.MobileContentScaleFactor=2
++CVars=sg.PostProcessQuality=3
 +CVars=sg.ShadowQuality=2
 +CVars=r.Mobile.AmbientOcclusionQuality=0
 +CVars=r.Mobile.PixelProjectedReflectionQuality=0
@@ -394,7 +284,7 @@
 
 [iPhone7 DeviceProfile]
 DeviceType=IOS
-BaseProfileName=iPhone6
+BaseProfileName=iPhone6S
 +CVars=ios.PhysicalScreenDensity=326
 +CVars=r.MobileContentScaleFactor=2
 +CVars=r.Mobile.AmbientOcclusionQuality=0
@@ -402,7 +292,7 @@
 
 [iPhone7Plus DeviceProfile]
 DeviceType=IOS
-BaseProfileName=iPhone6Plus
+BaseProfileName=iPhone6SPlus
 +CVars=ios.PhysicalScreenDensity=401
 +CVars=r.MobileContentScaleFactor=2
 +CVars=r.Mobile.AmbientOcclusionQuality=0
@@ -491,6 +381,51 @@
 BaseProfileName=iPhone11ProMax
 +CVars=ios.PhysicalScreenDensity=458
 
+[iPhone13Mini DeviceProfile]
+DeviceType=IOS
+BaseProfileName=iPhone11
++CVars=ios.PhysicalScreenDensity=476
+
+[iPhone13 DeviceProfile]
+DeviceType=IOS
+BaseProfileName=iPhone11
++CVars=ios.PhysicalScreenDensity=460
+
+[iPhone13Pro DeviceProfile]
+DeviceType=IOS
+BaseProfileName=iPhone11Pro
++CVars=ios.PhysicalScreenDensity=460
+
+[iPhone13ProMax DeviceProfile]
+DeviceType=IOS
+BaseProfileName=iPhone11ProMax
++CVars=ios.PhysicalScreenDensity=458
+
+[iPodTouch7 DeviceProfile]
+DeviceType=IOS
+BaseProfileName=iPhone7
++CVars=ios.PhysicalScreenDensity=326
++CVars=r.MobileContentScaleFactor=2
++CVars=r.Mobile.AmbientOcclusionQuality=0
++CVars=r.Mobile.PixelProjectedReflectionQuality=0
+
+[iPadAir2 DeviceProfile]
+DeviceType=IOS
+BaseProfileName=IOS
++CVars=ios.PhysicalScreenDensity=264
++CVars=sg.PostProcessQuality=3
++CVars=sg.ShadowQuality=2
++CVars=r.Mobile.AmbientOcclusionQuality=0
++CVars=r.Mobile.PixelProjectedReflectionQuality=0
+
+[iPadMini4 DeviceProfile]
+DeviceType=IOS
+BaseProfileName=IOS
++CVars=ios.PhysicalScreenDensity=401
++CVars=sg.PostProcessQuality=2
++CVars=r.Mobile.AmbientOcclusionQuality=0
++CVars=r.Mobile.PixelProjectedReflectionQuality=0
+
 [iPadPro DeviceProfile]
 DeviceType=IOS
 BaseProfileName=IOS
@@ -504,20 +439,6 @@
 +CVars=sg.PostProcessQuality=3
 +CVars=sg.AntiAliasingQuality=3
 
-[AppleTV DeviceProfile]
-DeviceType=IOS
-BaseProfileName=IOS
-+CVars=ios.PhysicalScreenDensity=0
-+CVars=sg.PostProcessQuality=3
-+CVars=sg.ShadowQuality=3
-+CVars=r.Mobile.AmbientOcclusionQuality=0
-+CVars=r.Mobile.PixelProjectedReflectionQuality=0
-
-[AppleTV4K DeviceProfile]
-DeviceType=IOS
-BaseProfileName=AppleTV
-+CVars=ios.PhysicalScreenDensity=0
-
 [iPadPro97 DeviceProfile]
 DeviceType=IOS
 BaseProfileName=IPadPro
@@ -543,6 +464,11 @@
 BaseProfileName=IPadPro
 +CVars=ios.PhysicalScreenDensity=264
 
+[iPadPro3_11 DeviceProfile]
+DeviceType=IOS
+BaseProfileName=IPadPro
++CVars=ios.PhysicalScreenDensity=264
+
 [iPad5 DeviceProfile]
 DeviceType=IOS
 BaseProfileName=IPadPro97
@@ -563,6 +489,11 @@
 BaseProfileName=IPadPro97
 +CVars=ios.PhysicalScreenDensity=264
 
+[iPad9 DeviceProfile]
+DeviceType=IOS
+BaseProfileName=IPadPro97
++CVars=ios.PhysicalScreenDensity=264
+
 [iPadPro2_129 DeviceProfile]
 DeviceType=IOS
 BaseProfileName=IPadPro
@@ -578,6 +509,11 @@
 BaseProfileName=IPadPro
 +CVars=ios.PhysicalScreenDensity=264
 
+[iPadPro5_129 DeviceProfile]
+DeviceType=IOS
+BaseProfileName=IPadPro
++CVars=ios.PhysicalScreenDensity=264
+
 [iPadAir3 DeviceProfile]
 DeviceType=IOS
 BaseProfileName=IPadPro
@@ -593,6 +529,25 @@
 BaseProfileName=IPadPro
 +CVars=ios.PhysicalScreenDensity=401
 
+[iPadMini6 DeviceProfile]
+DeviceType=IOS
+BaseProfileName=IPadPro
++CVars=ios.PhysicalScreenDensity=326
+
+[AppleTV DeviceProfile]
+DeviceType=IOS
+BaseProfileName=IOS
++CVars=ios.PhysicalScreenDensity=0
++CVars=sg.PostProcessQuality=3
++CVars=sg.ShadowQuality=3
++CVars=r.Mobile.AmbientOcclusionQuality=0
++CVars=r.Mobile.PixelProjectedReflectionQuality=0
+
+[AppleTV4K DeviceProfile]
+DeviceType=IOS
+BaseProfileName=AppleTV
++CVars=ios.PhysicalScreenDensity=0
+
 ;----------------
 ; Android devices
 
