Enable DDC2's BulkDataRegistry by default.
#rnx

diff --git a/Engine/Config/BaseEditor.ini b/Engine/Config/BaseEditor.ini
--- a/Engine/Config/BaseEditor.ini
+++ b/Engine/Config/BaseEditor.ini
@@ -311,7 +311,7 @@ MinFreeMemory=1024
 MaxUsedMemory=16384
 
 [CookSettings]
-BulkDataRegistryEnabled=false
+BulkDataRegistryEnabled=true
 ; HybridIterative uses the TargetDomain and BuildDependencies to decide when packages need to be recooked. If enabled,
 ; every cook is iterative for packages with allowed classes (IterativeClassAllowList) unless -noiterative is passed.
 HybridIterativeEnabled=false
