Stage Slate content in cooked editors

diff --git a/Engine/Config/BaseGame.ini b/Engine/Config/BaseGame.ini
--- a/Engine/Config/BaseGame.ini
+++ b/Engine/Config/BaseGame.ini
@@ -239,7 +239,6 @@ FrameInterpolationProcessor=Class'/Script/LiveLink.LiveLinkBasicFrameInterpolati
 +EngineExtraStageFiles=(Path="Content/Splash")
 +EngineExtraStageFiles=(Path="Content/EngineMaterials")
 +EngineExtraStageFiles=(Path="Content/EditorResources")
-+EngineExtraStageFiles=(Path="Content/Editor/Slate")
 +EngineExtraStageFiles=(Path="Content/EngineFonts")
 +EngineExtraStageFiles=(Path="Content/Maps/Templates")
 +EngineExtraStageFiles=(Path="Build", Files="Build.version", Recursive=false, NonUFS=true)
@@ -283,6 +282,10 @@ MapMode=Uncooked
 +DisabledPluginsInShipping=GameplayInsights
 +DisabledPluginsInShipping=TraceSourceFilters
 
+; stage extra Slate content
++EngineExtraStageFiles=(Path="Content/Editor/Slate")
++ProjectExtraStageFiles=(Path="Content/Editor/Slate")
+
 [CookedEditorSettings_CookedCooker]
 ; cooked cookers don't edit maps, so we can use cooked ones
 MapMode=Cooked
