Raise the minimum required driver versions for UE5.

diff --git a/Engine/Config/BaseHardware.ini b/Engine/Config/BaseHardware.ini
--- a/Engine/Config/BaseHardware.ini
+++ b/Engine/Config/BaseHardware.ini
@@ -33,39 +33,13 @@
 ;
 ;-----------------------------------------------------------------------------------------------------------------
 
-[GPU_NVIDIA Windows]
-+SuggestedDriverVersion="511.17"
-
-; Seems this driver was not setting uninitialized global variables
-; (static, non const) to 0 but leaving them uninitialized.
-; Also deny future driver versions until QA has verified them to work.
-; NOTE: UE-25096 has been fixed and verified by QA.
-;+DriverDenyList=(DriverVersion=">=361.43", Reason="UE-25096 Viewport flashes black and white when moving in the scene on latest NVIDIA drivers")
-
-; To prevent problems with older drivers.
-; We might have to adjust that further.
-+DriverDenyList=(DriverVersion="<=347.09", Reason="Crashes with Paragon content using more recent GPU features")
-
-; UE-35288, GPU hangs in 4.13 binary release.
-+DriverDenyList=(DriverVersion="==372.70", Reason="This driver version has many known stability issues")
-
-; UE-58673, GPU memory leak and hitching.
-+DriverDenyList=(DriverVersion="==388.31", Reason="This driver version has known stability issues")
-+DriverDenyList=(DriverVersion="==388.43", Reason="This driver version has known stability issues")
-+DriverDenyList=(DriverVersion="==388.71", Reason="This driver version has known stability issues")
-
-; UE-82937, Mobile GPUs on newer Windows 10 cannot handle swapchain creation
-+DriverDenyList=(DriverVersion="==376.54", Reason="This driver version has many known stability issues")
-+DriverDenyList=(DriverVersion="==382.05", Reason="This driver version has many known stability issues")
-+DriverDenyList=(DriverVersion="==388.57", Reason="This driver version has many known stability issues")
 
-; Driver with known stability issues
-+DriverDenyList=(DriverVersion="==461.40", Reason="This driver version has known stability issues")
 
-+DriverDenyList=(DriverVersion="==465.89", Reason="This driver version has known stability issues in DirectX 12")
+[GPU_NVIDIA Windows]
+SuggestedDriverVersion="511.79"
 
-; UE5 in Windows/Vulkan requires recent drivers for Nanite/Lumen and has other stability issues with older drivers (need >= 496.76)
-+DriverDenyList=(DriverVersion="<496.76", Reason="These driver versions have known stability issues and missing features in Vulkan", RHI="Vulkan")
+; To prevent problems with older drivers.
++DriverDenyList=(DriverVersion="<497.29", Reason="These driver versions have known stability issues and missing features")
 
 ;-----------------------------------------------------------------------------------------------------------------
 
@@ -82,18 +56,10 @@ SuggestedDriverVersion="510.47"
 ;-----------------------------------------------------------------------------------------------------------------
 
 [GPU_AMD Windows]
-+SuggestedDriverVersion="21.11.3"
-
-; From OR-27051
-+DriverDenyList=(DriverVersion="==15.300.1025.1001", Reason="Crashes with Paragon content using more recent GPU features")
+SuggestedDriverVersion="22.2.2"
 
-; 2017 drivers and older contribute heavily to swapchain creation errors.
-+DriverDenyList=(DriverVersion="<=22.19.662.4", Reason="Older drivers known to cause crashes with integrated laptop graphics on at least R5 and R6")
-
-+DriverDenyList=(DriverVersion="<=26.20.13031.15006", Reason="Driver crashes creating PSOs", RHI="D3D12")
-
-; UE5 in Windows/Vulkan requires recent drivers for Nanite/Lumen and has other stability issues with older drivers (need >= 21.11.3)
-+DriverDenyList=(DriverVersion="<30.0.13037.1003", Reason="These driver versions have known stability issues and missing features in Vulkan", RHI="Vulkan")
+; To prevent problems with older drivers.
++DriverDenyList=(DriverVersion="<30.0.14011.2006", Reason="These driver versions have known stability issues and missing features")
 
 ;-----------------------------------------------------------------------------------------------------------------
 
@@ -110,9 +76,9 @@ SuggestedDriverVersion="21.2.6"
 ;-----------------------------------------------------------------------------------------------------------------
 
 [GPU_Intel Windows]
-SuggestedDriverVersion="27.20.100.8935"
+SuggestedDriverVersion="30.0.101.1340"
 
 ;-----------------------------------------------------------------------------------------------------------------
 
 [GPU_Intel Linux]
-SuggestedDriverVersion="27.20.100.8935"
+SuggestedDriverVersion="30.0.101.1340"
