Added Android device profiles for Adreno 7xx and Xclipse 920
Support Vulkan by default
Disable Vulkan on Adreno 5xx

diff --git a/Engine/Config/BaseDeviceProfiles.ini b/Engine/Config/BaseDeviceProfiles.ini
--- a/Engine/Config/BaseDeviceProfiles.ini
+++ b/Engine/Config/BaseDeviceProfiles.ini
@@ -82,6 +82,8 @@
 +DeviceProfileNameAndTypes=Android_Adreno5xx,Android
 +DeviceProfileNameAndTypes=Android_Adreno6xx,Android
 +DeviceProfileNameAndTypes=Android_Adreno6xx_Vulkan,Android
++DeviceProfileNameAndTypes=Android_Adreno7xx,Android
++DeviceProfileNameAndTypes=Android_Adreno7xx_Vulkan,Android
 +DeviceProfileNameAndTypes=Android_Mali_T6xx,Android
 +DeviceProfileNameAndTypes=Android_Mali_T7xx,Android
 +DeviceProfileNameAndTypes=Android_Mali_T8xx,Android
@@ -96,6 +98,8 @@
 +DeviceProfileNameAndTypes=Android_Mali_G78_Vulkan,Android
 +DeviceProfileNameAndTypes=Android_Mali_G710,Android
 +DeviceProfileNameAndTypes=Android_Mali_G710_Vulkan,Android
++DeviceProfileNameAndTypes=Android_Xclipse_920,Android
++DeviceProfileNameAndTypes=Android_Xclipse_920_Vulkan,Android
 +DeviceProfileNameAndTypes=Android_Vulkan_SM5,Android
 +DeviceProfileNameAndTypes=Android_PowerVR_G6xxx,Android
 +DeviceProfileNameAndTypes=Android_PowerVR_GT7xxx,Android
@@ -553,22 +557,25 @@ BaseProfileName=AppleTV
 
 ; These rules are checked in turn and stop at the profile that matches all criteria
 [/Script/AndroidDeviceProfileSelector.AndroidDeviceProfileMatchingRules]
-; Specific Android devices
+; Match Quest devices directly
 +MatchProfile=(Profile="Oculus_Quest",Match=((SourceType=SRC_DeviceMake,CompareType=CMP_Equal,MatchString="Oculus"), (SourceType=SRC_HMDSystemName,CompareType=CMP_Equal,MatchString="Oculus Quest")))
 +MatchProfile=(Profile="Oculus_Quest2",Match=((SourceType=SRC_DeviceMake,CompareType=CMP_Equal,MatchString="Oculus"), (SourceType=SRC_HMDSystemName,CompareType=CMP_Equal,MatchString="Oculus Quest2")))
-; Standard android
+
+; Adreno-based devices
 +MatchProfile=(Profile="Android_Adreno4xx",Match=((SourceType=SRC_GpuFamily,CompareType=CMP_Regex,MatchString="Adreno \\(TM\\) 4[0-9][0-9]")))
-+MatchProfile=(Profile="Android_Adreno5xx_Low_Vulkan",Match=((SourceType=SRC_GpuFamily,CompareType=CMP_Regex,MatchString="Adreno \\(TM\\) 5[0-1][0-9]"),(SourceType=SRC_AndroidVersion, CompareType=CMP_Regex,MatchString="([0-9]+).*"),(SourceType=SRC_PreviousRegexMatch,CompareType=CMP_GreaterEqual,MatchString="10"),(SourceType=SRC_VulkanAvailable,CompareType=CMP_Equal,MatchString="true")))
 +MatchProfile=(Profile="Android_Adreno5xx_Low",Match=((SourceType=SRC_GpuFamily,CompareType=CMP_Regex,MatchString="Adreno \\(TM\\) 5[0-1][0-9]")))
-+MatchProfile=(Profile="Android_Adreno5xx_Vulkan",Match=((SourceType=SRC_GpuFamily,CompareType=CMP_Regex,MatchString="Adreno \\(TM\\) 5[0-9][0-9]"),(SourceType=SRC_AndroidVersion, CompareType=CMP_Regex,MatchString="([0-9]+).*"),(SourceType=SRC_PreviousRegexMatch,CompareType=CMP_GreaterEqual,MatchString="10"),(SourceType=SRC_VulkanAvailable,CompareType=CMP_Equal,MatchString="true")))
 +MatchProfile=(Profile="Android_Adreno5xx",Match=((SourceType=SRC_GpuFamily,CompareType=CMP_Regex,MatchString="Adreno \\(TM\\) 5[0-9][0-9]")))
 ; Example profiles for Vulkan SM5
 ;+MatchProfile=(Profile="Android_Vulkan_SM5",Match=((SourceType=SRC_GpuFamily,CompareType=CMP_Regex,MatchString="Adreno \\(TM\\) 6[0-9][0-9]"),(SourceType=SRC_AndroidVersion, CompareType=CMP_Regex,MatchString="([0-9]+).*"),(SourceType=SRC_PreviousRegexMatch,CompareType=CMP_GreaterEqual,MatchString="9"),(SourceType=SRC_VulkanAvailable,CompareType=CMP_Equal,MatchString="true")))
 ;+MatchProfile=(Profile="Android_Vulkan_SM5_Mali",Match=((SourceType=SRC_GpuFamily,CompareType=CMP_Regex,MatchString="^Mali\\-G7[67]"),(SourceType=SRC_VulkanAvailable,CompareType=CMP_Equal,MatchString="true")))
 +MatchProfile=(Profile="Android_Adreno6xx_Vulkan",Match=((SourceType=SRC_GpuFamily,CompareType=CMP_Regex,MatchString="Adreno \\(TM\\) 6[0-9][0-9]"),(SourceType=SRC_AndroidVersion, CompareType=CMP_Regex,MatchString="([0-9]+).*"),(SourceType=SRC_PreviousRegexMatch,CompareType=CMP_GreaterEqual,MatchString="9"),(SourceType=SRC_VulkanAvailable,CompareType=CMP_Equal,MatchString="true")))
 +MatchProfile=(Profile="Android_Adreno6xx",Match=((SourceType=SRC_GpuFamily,CompareType=CMP_Regex,MatchString="Adreno \\(TM\\) 6[0-9][0-9]")))
++MatchProfile=(Profile="Android_Adreno7xx_Vulkan",Match=((SourceType=SRC_GpuFamily,CompareType=CMP_Regex,MatchString="Adreno \\(TM\\) 7[0-9][0-9]"),(SourceType=SRC_AndroidVersion, CompareType=CMP_Regex,MatchString="([0-9]+).*"),(SourceType=SRC_PreviousRegexMatch,CompareType=CMP_GreaterEqual,MatchString="9"),(SourceType=SRC_VulkanAvailable,CompareType=CMP_Equal,MatchString="true")))
++MatchProfile=(Profile="Android_Adreno7xx",Match=((SourceType=SRC_GpuFamily,CompareType=CMP_Regex,MatchString="Adreno \\(TM\\) 7[0-9][0-9]")))
 ; Enable Vulkan on any unknown Adreno running Android 10 or later
 +MatchProfile=(Profile="Android_Unknown_Vulkan",Match=((SourceType=SRC_GpuFamily,CompareType=CMP_Regex,MatchString="Adreno"),(SourceType=SRC_AndroidVersion, CompareType=CMP_Regex,MatchString="([0-9]+).*"),(SourceType=SRC_PreviousRegexMatch,CompareType=CMP_GreaterEqual,MatchString="10"),(SourceType=SRC_VulkanAvailable,CompareType=CMP_Equal,MatchString="true")))
+
+; Mali-based devices
 +MatchProfile=(Profile="Android_Mali_T6xx",Match=((SourceType=SRC_GpuFamily,CompareType=CMP_Regex,MatchString="^Mali\\-T6")))
 +MatchProfile=(Profile="Android_Mali_T7xx",Match=((SourceType=SRC_GpuFamily,CompareType=CMP_Regex,MatchString="^Mali\\-T7")))
 +MatchProfile=(Profile="Android_Mali_T8xx",Match=((SourceType=SRC_GpuFamily,CompareType=CMP_Regex,MatchString="^Mali\\-T8")))
@@ -585,12 +592,21 @@ BaseProfileName=AppleTV
 +MatchProfile=(Profile="Android_Mali_G710",Match=((SourceType=SRC_GpuFamily,CompareType=CMP_Regex,MatchString="^Mali\\-G710")))
 ; Enable Vulkan on any unknown Mali running Android 10 or later
 +MatchProfile=(Profile="Android_Unknown_Vulkan",Match=((SourceType=SRC_GpuFamily,CompareType=CMP_Regex,MatchString="^Mali"),(SourceType=SRC_AndroidVersion, CompareType=CMP_Regex,MatchString="([0-9]+).*"),(SourceType=SRC_PreviousRegexMatch,CompareType=CMP_GreaterEqual,MatchString="10"),(SourceType=SRC_VulkanAvailable,CompareType=CMP_Equal,MatchString="true")))
+
+; Samsung Exynos / Xclipse GPU-based devices use Vulkan by default.
++MatchProfile=(Profile="Android_Xclipse_920_Vulkan",Match=((SourceType=SRC_GpuFamily,CompareType=CMP_Regex,MatchString="Samsung Xclipse 920"),(SourceType=SRC_VulkanAvailable,CompareType=CMP_Equal,MatchString="true")))
++MatchProfile=(Profile="Android_Xclipse_920",Match=((SourceType=SRC_GpuFamily,CompareType=CMP_Regex,MatchString="Samsung Xclipse 920")))
+
+; PowerVR based devices
 +MatchProfile=(Profile="Android_PowerVR_G6xxx",Match=((SourceType=SRC_GpuFamily,CompareType=CMP_Regex,MatchString="PowerVR Rogue G6[0-9]+")))
 +MatchProfile=(Profile="Android_PowerVR_GT7xxx",Match=((SourceType=SRC_GpuFamily,CompareType=CMP_Regex,MatchString="PowerVR Rogue GT7[0-9]+")))
 +MatchProfile=(Profile="Android_PowerVR_GE8xxx",Match=((SourceType=SRC_GpuFamily,CompareType=CMP_Regex,MatchString="PowerVR Rogue GE8[0-9]+")))
 +MatchProfile=(Profile="Android_PowerVR_GM9xxx_Vulkan",Match=((SourceType=SRC_GpuFamily,CompareType=CMP_Regex,MatchString="PowerVR Rogue GM9[0-9]+"),(SourceType=SRC_VulkanAvailable,CompareType=CMP_Equal,MatchString="true")))
 +MatchProfile=(Profile="Android_PowerVR_GM9xxx",Match=((SourceType=SRC_GpuFamily,CompareType=CMP_Regex,MatchString="PowerVR Rogue GM9[0-9]+")))
+; Tegra devices
+
 +MatchProfile=(Profile="Android_TegraK1",Match=((SourceType=SRC_GpuFamily,CompareType=CMP_Equal,MatchString="NVIDIA Tegra"),(SourceType=SRC_GlVersion,CompareType=CMP_Regex,MatchString="^OpenGL ES 3\\.")))
+
 ; Android_Default is the profile of last resort
 
 [Android DeviceProfile]
@@ -613,7 +629,7 @@ PreviewAllowlistCVars=sg.
 [Android_Preview_OpenGL DeviceProfile]
 DeviceType=Android
 BaseProfileName=Android
-; FX.AllowGPUSorting is hardcoded to 0 in OpenGLDevice.cpp to avoid driver bugs
+; FX.AllowGPUSorting is hardcoded to 0 in OpenGLDevice.cpp to avoid driver bugs. Emulate this behavior in preview.
 +CVars=FX.AllowGPUSorting=0
 
 [Android_Preview_Vulkan DeviceProfile]
@@ -709,29 +725,40 @@ BaseProfileName=Android_Low
 DeviceType=Android
 BaseProfileName=Android_Low
 
-[Android_Adreno5xx_Low_Vulkan DeviceProfile]
+[Android_Adreno5xx DeviceProfile]
+DeviceType=Android
+BaseProfileName=Android_Mid
+
+[Android_Adreno6xx DeviceProfile]
+DeviceType=Android
+BaseProfileName=Android_High
+
+[Android_Adreno6xx_Vulkan DeviceProfile]
 DeviceType=Android
-BaseProfileName=Android_Adreno5xx_Low
+BaseProfileName=Android_Adreno6xx
 ; enable Vulkan
 +CVars=r.Android.DisableVulkanSupport=0
 
-[Android_Adreno5xx DeviceProfile]
+[Android_Adreno7xx DeviceProfile]
 DeviceType=Android
-BaseProfileName=Android_Mid
+BaseProfileName=Android_High
 
-[Android_Adreno5xx_Vulkan DeviceProfile]
+[Android_Adreno7xx_Vulkan DeviceProfile]
 DeviceType=Android
-BaseProfileName=Android_Adreno5xx
+BaseProfileName=Android_Adreno7xx
 ; enable Vulkan
 +CVars=r.Android.DisableVulkanSupport=0
 
-[Android_Adreno6xx DeviceProfile]
+;----------------
+; Xclipse devices
+
+[Android_Xclipse_920 DeviceProfile]
 DeviceType=Android
 BaseProfileName=Android_High
 
-[Android_Adreno6xx_Vulkan DeviceProfile]
+[Android_Xclipse_920_Vulkan DeviceProfile]
 DeviceType=Android
-BaseProfileName=Android_Adreno6xx
+BaseProfileName=Android_Xclipse_920
 ; enable Vulkan
 +CVars=r.Android.DisableVulkanSupport=0
 
diff --git a/Engine/Config/BaseEngine.ini b/Engine/Config/BaseEngine.ini
--- a/Engine/Config/BaseEngine.ini
+++ b/Engine/Config/BaseEngine.ini
@@ -2604,7 +2604,7 @@ bSupportAdMob=true
 bBuildForArm64=true
 bBuildForX8664=false
 bBuildForES31=true
-bSupportsVulkan=false
+bSupportsVulkan=true
 bSupportsVulkanSM5=false
 bDetectVulkanByDefault=true
 bSplitIntoSeparateApks=false
