Device profile change to fix rendering artifacts caused by SkyAtmosphere on PowerVR mobile devices.

diff --git a/Engine/Config/BaseDeviceProfiles.ini b/Engine/Config/BaseDeviceProfiles.ini
--- a/Engine/Config/BaseDeviceProfiles.ini
+++ b/Engine/Config/BaseDeviceProfiles.ini
@@ -772,18 +772,22 @@ BaseProfileName=Android_Xclipse_920
 [Android_PowerVR_G6xxx DeviceProfile]
 DeviceType=Android
 BaseProfileName=Android_Low
++CVars=r.SkyAtmosphere.LUT32=1
 
 [Android_PowerVR_GT7xxx DeviceProfile]
 DeviceType=Android
 BaseProfileName=Android_Mid
++CVars=r.SkyAtmosphere.LUT32=1
 
 [Android_PowerVR_GE8xxx DeviceProfile]
 DeviceType=Android
 BaseProfileName=Android_Mid
++CVars=r.SkyAtmosphere.LUT32=1
 
 [Android_PowerVR_GM9xxx DeviceProfile]
 DeviceType=Android
 BaseProfileName=Android_High
++CVars=r.SkyAtmosphere.LUT32=1
 
 [Android_PowerVR_GM9xxx_Vulkan DeviceProfile]
 DeviceType=Android
