Enables r.TSR.History.HighFrequencyOnly=1 in medium and low AA scalability

diff --git a/Engine/Config/BaseScalability.ini b/Engine/Config/BaseScalability.ini
--- a/Engine/Config/BaseScalability.ini
+++ b/Engine/Config/BaseScalability.ini
@@ -42,6 +42,7 @@ r.TemporalAA.Quality=0
 r.TSR.History.R11G11B10=1
 r.TSR.History.ScreenPercentage=100
 r.TSR.History.UpdateQuality=0
+r.TSR.History.HighFrequencyOnly=1
 r.TSR.ShadingRejection.SpatialFilter=0
 r.TSR.Velocity.Extrapolation=0
 r.TSR.Velocity.HoleFill=0
@@ -53,6 +54,7 @@ r.TemporalAA.Quality=1
 r.TSR.History.R11G11B10=1
 r.TSR.History.ScreenPercentage=100
 r.TSR.History.UpdateQuality=1
+r.TSR.History.HighFrequencyOnly=1
 r.TSR.ShadingRejection.SpatialFilter=1
 r.TSR.Velocity.Extrapolation=0
 r.TSR.Velocity.HoleFill=1
@@ -64,6 +66,7 @@ r.TemporalAA.Quality=1
 r.TSR.History.R11G11B10=1
 r.TSR.History.ScreenPercentage=100
 r.TSR.History.UpdateQuality=2
+r.TSR.History.HighFrequencyOnly=0
 r.TSR.ShadingRejection.SpatialFilter=1
 r.TSR.Velocity.Extrapolation=0
 r.TSR.Velocity.HoleFill=1
@@ -75,6 +78,7 @@ r.TemporalAA.Quality=2
 r.TSR.History.R11G11B10=1
 r.TSR.History.ScreenPercentage=100
 r.TSR.History.UpdateQuality=3
+r.TSR.History.HighFrequencyOnly=0
 r.TSR.ShadingRejection.SpatialFilter=2
 r.TSR.Velocity.Extrapolation=1
 r.TSR.Velocity.HoleFill=1
@@ -86,6 +90,7 @@ r.TemporalAA.Quality=2
 r.TSR.History.R11G11B10=0
 r.TSR.History.ScreenPercentage=200
 r.TSR.History.UpdateQuality=3
+r.TSR.History.HighFrequencyOnly=0
 r.TSR.ShadingRejection.SpatialFilter=2
 r.TSR.Velocity.Extrapolation=1
 r.TSR.Velocity.HoleFill=1
