Deprecating MatchDevice in Sound Wave Quality settings

diff --git a/Engine/Config/BaseEngine.ini b/Engine/Config/BaseEngine.ini
--- a/Engine/Config/BaseEngine.ini
+++ b/Engine/Config/BaseEngine.ini
@@ -997,6 +997,9 @@ DefaultWorldPartitionHLODLayer=/Engine/Maps/Templates/HLODs/HLODLayer_Instanced.
 +FunctionRedirects=(OldName="SoundGeneratorParameterInterface.SetIntParameter", NewName="AudioParameterControllerInterface.SetIntParameter")
 +ClassRedirects=(OldName="/Script/AudioExtensions.AudioParameterInterface", NewName="/Script/AudioExtensions.AudioParameterControllerInterface")
 
+; Deprecating 'MatchDevice' in Sound Wave Quality Enum
++EnumRedirects=(OldName="ESoundwaveSampleRateSettings",ValueChanges=(("ESoundwaveSampleRateSettings::MatchDevice","ESoundwaveSampleRateSettings::Max")) )
+
 ; Updating the gameplay  cue notify to have correct property name
 +PropertyRedirects=(OldName="GameplayCueNotify_SoundInfo.SoundCue",NewName="GameplayCueNotify_SoundInfo.Sound")
 
diff --git a/Engine/Source/Runtime/AudioPlatformConfiguration/Public/AudioCompressionSettings.h b/Engine/Source/Runtime/AudioPlatformConfiguration/Public/AudioCompressionSettings.h
--- a/Engine/Source/Runtime/AudioPlatformConfiguration/Public/AudioCompressionSettings.h
+++ b/Engine/Source/Runtime/AudioPlatformConfiguration/Public/AudioCompressionSettings.h
@@ -11,8 +11,7 @@ enum class ESoundwaveSampleRateSettings : uint8
 	Medium,
 	Low,
 	Min,
-	// Use this setting to resample soundwaves to the device's sample rate to avoid having to perform sample rate conversion at runtime.
-	MatchDevice
+	MatchDevice_DEPRECATED
 };
 
 /************************************************************************/
