Prevent our PipelineCaches directories from being trimmed in pak memory optimization. We need to query these to tell when is appropriate to kick off shader loads when installing data.
#tests Android Test Client

diff --git a/Engine/Config/BaseEngine.ini b/Engine/Config/BaseEngine.ini
--- a/Engine/Config/BaseEngine.ini
+++ b/Engine/Config/BaseEngine.ini
@@ -2542,7 +2542,9 @@ bDoUseOodleDespiteNoPluginCompression = true
 +DirectoryIndexKeepFiles="*/Content/Movies/*"
 +DirectoryIndexKeepFiles="*/Plugins/*.uplugin*"
 +DirectoryIndexKeepFiles="*/*.ini*"
++DirectoryIndexKeepFiles="*/*.upipelinecache*"
 +DirectoryIndexKeepEmptyDirectories=*/Plugins/*
++DirectoryIndexKeepEmptyDirectories="*/PipelineCaches/*"
 +IndexValidationIgnore="*/Saved/Logs/*"
 
 [/Script/GameplayDebugger.GameplayDebuggingReplicator]
