Adjustments to the allow/block types for editor domain and hybrid iterative.

diff --git a/Engine/Config/BaseEditor.ini b/Engine/Config/BaseEditor.ini
--- a/Engine/Config/BaseEditor.ini
+++ b/Engine/Config/BaseEditor.ini
@@ -370,6 +370,8 @@ SaveUnversioned=true
 +ClassBlockList=/Script/Engine.StaticMesh
 +ClassBlockList=/Script/Engine.Blueprint
 +ClassBlockList=/Script/Engine.UserDefinedStruct
++ClassBlockList=/Script/CommonUI.CommonGenericInputActionDataTable
++ClassBlockList=/Script/Engine.DataTable
 
 ; Known bad files that are problematic if all of them come from EditorDomain. Blocking any one of these is sufficient to prevent the bug.
 +PackageBlockList=/Niagara/Functions/RandomRangeFloat
@@ -418,6 +420,7 @@ IterativeClassAllowListEnabled=true
 IterativeClassAllowList=/Script/Engine.Texture2D
 +IterativeClassAllowList=/Script/CoreUObject.MetaData
 +IterativeClassAllowList=/Script/Engine.AssetImportData
++IterativeClassAllowList=/Script/Engine.SoundWave
 
 [CookPlatformDataCacheSettings]
 MaterialInstance=100
