Add HybridIterativeEnabled to deny list for settings changes that invalidate cooked output.

diff --git a/Engine/Config/BaseEditor.ini b/Engine/Config/BaseEditor.ini
--- a/Engine/Config/BaseEditor.ini
+++ b/Engine/Config/BaseEditor.ini
@@ -337,6 +337,7 @@ bUseFullGCAssetClassNames=true
 +CookOnTheFlyConfigSettingDenyList=*.UserSettings
 +CookOnTheFlyConfigSettingDenyList=*.GameUserSettings
 +CookOnTheFlyConfigSettingDenyList=*.Editor:CookSettings:CookOnTheFlyConfigSettingDenyList
++CookOnTheFlyConfigSettingDenyList=*.Editor:CookSettings:HybridIterativeEnabled
 +CookOnTheFlyConfigSettingDenyList=*.EditorSettings:/Script/UnrealEd.EditorSettings
 +CookOnTheFlyConfigSettingDenyList=*.EditorSettings
 +CookOnTheFlyConfigSettingDenyList=*.EditorLayout
