Use Build.bat/sh as a wrapper for UBT as we want to setup the dotnet bundled version vs depending on the system

diff --git a/Engine/Config/BaseEngine.ini b/Engine/Config/BaseEngine.ini
--- a/Engine/Config/BaseEngine.ini
+++ b/Engine/Config/BaseEngine.ini
@@ -2970,7 +2970,7 @@ DefaultChaosSolverActorClass=/Script/ChaosSolverEngine.ChaosSolverActor
 ;+Pools=(Formats=(PF_DXT5, PF_DXT5), SizeInMegabyte=100)
 
 [PlatformPaths]
-UnrealBuildTool=Engine/Binaries/DotNET/UnrealBuildTool/UnrealBuildTool.exe
+UnrealBuildTool=Engine/Build/BatchFiles/Build.bat
 
 [/Script/InterchangeEngine.InterchangeProjectSettings]
 PipelineStacks=(("Generic", (Pipelines=("/Script/InterchangePipelines.InterchangeGenericAssetsPipeline"))),("Scene", (Pipelines=("/Script/InterchangePipelines.InterchangeGenericAssetsPipeline","/Script/InterchangePipelines.InterchangeGenericLevelPipeline"))))
diff --git a/Engine/Config/Linux/BaseLinuxEngine.ini b/Engine/Config/Linux/BaseLinuxEngine.ini
--- a/Engine/Config/Linux/BaseLinuxEngine.ini
+++ b/Engine/Config/Linux/BaseLinuxEngine.ini
@@ -1,5 +1,5 @@
 [PlatformPaths]
-UnrealBuildTool=Engine/Binaries/DotNET/UnrealBuildTool/UnrealBuildTool
+UnrealBuildTool=Engine/Build/BatchFiles/Linux/Build.sh
 
 [OnlineServices.EOS.Auth]
-DefaultExternalCredentialTypeStr=Steam
\ No newline at end of file
+DefaultExternalCredentialTypeStr=Steam
diff --git a/Engine/Config/Mac/BaseMacEngine.ini b/Engine/Config/Mac/BaseMacEngine.ini
--- a/Engine/Config/Mac/BaseMacEngine.ini
+++ b/Engine/Config/Mac/BaseMacEngine.ini
@@ -5,7 +5,7 @@ AudioCaptureModuleName=AudioCaptureRtAudio
 LocalPlatformName=MAC
 
 [PlatformPaths]
-UnrealBuildTool=Engine/Binaries/DotNET/UnrealBuildTool/UnrealBuildTool
+UnrealBuildTool=Engine/Build/BatchFiles/Mac/Build.sh
 
 [OnlineServices.EOS.Auth]
-DefaultExternalCredentialTypeStr=Steam
\ No newline at end of file
+DefaultExternalCredentialTypeStr=Steam
