Build fix for Lyra: forgot a property redirect for a renamed parameter, which was caused by CL 19726070.

Additionally, added a missing 5.1 section for redirects.

diff --git a/Engine/Config/BaseEngine.ini b/Engine/Config/BaseEngine.ini
--- a/Engine/Config/BaseEngine.ini
+++ b/Engine/Config/BaseEngine.ini
@@ -1054,7 +1054,11 @@ DefaultWorldPartitionHLODLayer=/Engine/Maps/Templates/HLODs/HLODLayer_Instanced.
 +FunctionRedirects=(OldName="KismetMathLibrary.EqualEqual_FloatFloat", NewName="EqualEqual_DoubleDouble")
 +FunctionRedirects=(OldName="KismetMathLibrary.NotEqual_FloatFloat", NewName="NotEqual_DoubleDouble")
 
+;5.1
+
+; Blueprint real number support
 +FunctionRedirects=(OldName="KismetTextLibrary.Conv_FloatToText", NewName="Conv_DoubleToText")
++PropertyRedirects=(OldName="KismetTextLibrary.Conv_DoubleToText.InDouble", NewName="Conv_DoubleToText.Value")
 
 [CoreUObject.Metadata]
 ; Note: UnrealHeaderTool should be rerun after making changes to MetadataRedirects to catch any keys specified in class headers
