Rename Zen DDC graph nodes to avoid using the "Local" node name.  Also removed one debug graph that isn't used.

#rnx

diff --git a/Engine/Config/BaseEngine.ini b/Engine/Config/BaseEngine.ini
--- a/Engine/Config/BaseEngine.ini
+++ b/Engine/Config/BaseEngine.ini
@@ -2218,14 +2218,14 @@ EnterprisePak=(Type=ReadPak, Filename=../../../Enterprise/DerivedDataCache/Compr
 
 [ZenOnlyDDC]
 Root=(Type=KeyLength, Length=120, Inner=AsyncPut)
-AsyncPut=(Type=AsyncPut, Inner=Local)
-Local=(Type=Zen, Namespace="ue4.ddc")
+AsyncPut=(Type=AsyncPut, Inner=Zen)
+Zen=(Type=Zen, Namespace="ue4.ddc")
 
 [ZenDDC]
 Root=(Type=KeyLength, Length=120, Inner=AsyncPut)
 AsyncPut=(Type=AsyncPut, Inner=Hierarchy)
-Hierarchy=(Type=Hierarchical, Inner=Pak, Inner=EnginePak, Inner=Local, Inner=Shared)
-Local=(Type=Zen, Namespace="ue4.ddc")
+Hierarchy=(Type=Hierarchical, Inner=Pak, Inner=EnginePak, Inner=Zen, Inner=Shared)
+Zen=(Type=Zen, Namespace="ue4.ddc")
 Shared=(Type=FileSystem, UnusedFileAge=10, FoldersToClean=10, ConsiderSlowAt=70, Path=?EpicDDC, EnvPathOverride=UE-SharedDataCachePath, EditorOverrideSetting=SharedDerivedDataCache, CommandLineOverride=SharedDataCachePath)
 Pak=(Type=ReadPak, Filename="%GAMEDIR%DerivedDataCache/DDC.ddp")
 EnginePak=(Type=ReadPak, Filename="%ENGINEDIR%DerivedDataCache/DDC.ddp")
