Expose r.Shaders.ExtraData in ConsoleVariables.ini

diff --git a/Engine/Config/ConsoleVariables.ini b/Engine/Config/ConsoleVariables.ini
--- a/Engine/Config/ConsoleVariables.ini
+++ b/Engine/Config/ConsoleVariables.ini
@@ -72,6 +72,8 @@ r.DumpShaderDebugInfo=2
 ;r.Shaders.Optimize=0
 ; When this is enabled, shaders will have extra debugging info. This could change patch sizes, uniqueness, etc and will recompile the shaders
 ;r.Shaders.Symbols=1
+; When this is enabled, ShaderName field of FRHIShader will be populated (Development and Debug builds only)
+;r.Shaders.ExtraData=1
 ; Uncomment to skip shader compression. Can save a significant time when using debug shaders.
 ;r.Shaders.SkipCompression=1
 ; Uncomment and change to set a custom folder to write shader symbols to when they are generated.
