- Disabling USD plugins in cooked editors (monolithic) which fail to compile due to some USE_USD mismatches

diff --git a/Engine/Config/BaseGame.ini b/Engine/Config/BaseGame.ini
--- a/Engine/Config/BaseGame.ini
+++ b/Engine/Config/BaseGame.ini
@@ -254,6 +254,10 @@ bStagePlatformDirs=true
 bStageUAT=true
 bIsForExternalDistribution=false
 
+; this won't compile in monolithic editor without forcing ANSI allocator
++DisabledPlugins=USDStageEditor
++DisabledPlugins=USDImporter
+
 
 ; settings related to building the cooked editor against an existing release, as a DLC plugin
 bBuildAgainstRelease=false
