Expand the IterativeClassAllowList to include SoundCue, StaticMesh, and MaterialInstanceConstant (and commonly co-packaged classes).  For further testing with hybrid iterative cooking.

#rnx

diff --git a/Engine/Config/BaseEditor.ini b/Engine/Config/BaseEditor.ini
--- a/Engine/Config/BaseEditor.ini
+++ b/Engine/Config/BaseEditor.ini
@@ -418,10 +418,44 @@ SaveUnversioned=true
 ; +IterativeClassAllowList=/Long/PackageName.ClassName
 ; +IterativeClassBlockList=/Long/PackageName.ClassName
 IterativeClassAllowListEnabled=true
+
 IterativeClassAllowList=/Script/Engine.Texture2D
++IterativeClassAllowList=/Script/Engine.TextureCube
++IterativeClassAllowList=/Script/Engine.SoundWave
++IterativeClassAllowList=/Script/Engine.DialogueWave
+
+; Universally co-packaged classes
 +IterativeClassAllowList=/Script/CoreUObject.MetaData
 +IterativeClassAllowList=/Script/Engine.AssetImportData
-+IterativeClassAllowList=/Script/Engine.SoundWave
++IterativeClassAllowList=/Script/UnrealEd.SceneThumbnailInfo
++IterativeClassAllowList=/Script/UnrealEd.SceneThumbnailInfoWithPrimitive
+
+; SoundCue and commonly co-packaged classes
++IterativeClassAllowList=/Script/Engine.SoundCue
++IterativeClassAllowList=/Script/AudioEditor.SoundCueGraph
++IterativeClassAllowList=/Script/AudioEditor.SoundCueGraphNode_Root
++IterativeClassAllowList=/Script/AudioEditor.SoundCueGraphNode
++IterativeClassAllowList=/Script/Engine.SoundNodeDialoguePlayer
++IterativeClassAllowList=/Script/Engine.SoundNodeWavePlayer
++IterativeClassAllowList=/Script/Engine.SoundNodeRandom
++IterativeClassAllowList=/Script/Engine.SoundNodeModulator
++IterativeClassAllowList=/Script/Engine.SoundNodeMixer
++IterativeClassAllowList=/Script/Engine.SoundNodeDistanceCrossFade
++IterativeClassAllowList=/Script/Engine.SoundNodeSoundClass
++IterativeClassAllowList=/Script/UnrealEd.EdGraphNode_Comment
++IterativeClassAllowList=/Script/Engine.SoundNodeAttenuation
+
+; StaticMesh and commonly co-packaged classes
++IterativeClassAllowList=/Script/Engine.StaticMesh
++IterativeClassAllowList=/Script/Engine.BodySetup
++IterativeClassAllowList=/Script/NavigationSystem.NavCollision
++IterativeClassAllowList=/Script/Engine.StaticMeshDescriptionBulkData
++IterativeClassAllowList=/Script/UnrealEd.FbxStaticMeshImportData
++IterativeClassAllowList=/Script/Engine.StaticMeshSocket
+
+; MaterialInstanceConstant and commonly co-packaged classes
++IterativeClassAllowList=/Script/Engine.MaterialInstanceConstant
++IterativeClassAllowList=/Script/Engine.SceneThumbnailInfoWithPrimitive
 
 [CookPlatformDataCacheSettings]
 MaterialInstance=100
