TimeOfDay: Add default tracks to TimeOfDaySequenceEditor settings.
Add default editor preferences for TimeOfDaySequenceEditor.

diff --git a/Engine/Config/BaseEditorPerProjectUserSettings.ini b/Engine/Config/BaseEditorPerProjectUserSettings.ini
--- a/Engine/Config/BaseEditorPerProjectUserSettings.ini
+++ b/Engine/Config/BaseEditorPerProjectUserSettings.ini
@@ -880,6 +880,14 @@
 LoopMode=SLM_Loop
 bCleanPlaybackMode=false
 
+[TimeOfDaySequenceEditor SequencerSettings]
+bShowRangeSlider=true
+bKeepPlayRangeInSectionBounds=false
+ZeroPadFrames=4
+bInfiniteKeyAreas=true
+bAutoSetTrackDefaults=true
+FrameNumberDisplayFormat=Frames
+
 [/Script/LevelSequenceEditor.LevelSequenceEditorSettings]
 +TrackSettings=(MatchingActorClass=/Script/Engine.StaticMeshActor,DefaultTracks=(/Script/MovieSceneTracks.MovieScene3DTransformTrack))
 +TrackSettings=(MatchingActorClass=/Script/Engine.SkeletalMeshActor,DefaultTracks=(/Script/MovieSceneTracks.MovieScene3DTransformTrack,/Script/MovieSceneTracks.MovieSceneSkeletalAnimationTrack))
@@ -894,6 +902,13 @@
 +TrackSettings=(MatchingActorClass=/Script/CinematicCamera.CineCameraActor,DefaultPropertyTracks=((ComponentPath="CameraComponent",PropertyPath="CurrentFocalLength"),(ComponentPath="CameraComponent",PropertyPath="FocusSettings.ManualFocusDistance"),(ComponentPath="CameraComponent",PropertyPath="CurrentAperture")),ExcludeDefaultPropertyTracks=((ComponentPath="CameraComponent",PropertyPath="FieldOfView")))
 +TrackSettings=(MatchingActorClass=/Script/Engine.Light,DefaultTracks=(None),DefaultPropertyTracks=((ComponentPath="LightComponent0",PropertyPath="Intensity"),(ComponentPath="LightComponent0",PropertyPath="LightColor")))
 
+[/Script/TimeOfDayEditor.TimeOfDayEditorSettings]
++TrackSettings=(MatchingActorClass=/Script/Engine.StaticMeshActor,DefaultTracks=(/Script/MovieSceneTracks.MovieScene3DTransformTrack))
++TrackSettings=(MatchingActorClass=/Script/Engine.SkeletalMeshActor,DefaultTracks=(/Script/MovieSceneTracks.MovieScene3DTransformTrack,/Script/MovieSceneTracks.MovieSceneSkeletalAnimationTrack))
++TrackSettings=(MatchingActorClass=/Script/Engine.CameraActor,DefaultTracks=(/Script/MovieSceneTracks.MovieScene3DTransformTrack),DefaultPropertyTracks=((ComponentPath="CameraComponent",PropertyPath="FieldOfView")))
++TrackSettings=(MatchingActorClass=/Script/CinematicCamera.CineCameraActor,DefaultPropertyTracks=((ComponentPath="CameraComponent",PropertyPath="CurrentFocalLength"),(ComponentPath="CameraComponent",PropertyPath="FocusSettings.ManualFocusDistance"),(ComponentPath="CameraComponent",PropertyPath="CurrentAperture")),ExcludeDefaultPropertyTracks=((ComponentPath="CameraComponent",PropertyPath="FieldOfView")))
++TrackSettings=(MatchingActorClass=/Script/Engine.Light,DefaultTracks=(None),DefaultPropertyTracks=((ComponentPath="LightComponent0",PropertyPath="Intensity"),(ComponentPath="LightComponent0",PropertyPath="LightColor")))
+
 [/Script/MovieSceneTools.MovieSceneToolsProjectSettings]
 +FbxSettings=(FbxPropertyName="FieldOfView", PropertyPath=(ComponentName="CameraComponent",PropertyName="FieldOfView"))
 +FbxSettings=(FbxPropertyName="FocalLength", PropertyPath=(ComponentName="CameraComponent",PropertyName="CurrentFocalLength"))
