Enable the package trailer system for the entire engine

#rnx

- This will change how we store editor bulkdata payloads in the package file and has been running on various branches for half a year so should be quite safe to turn on everywhere.

diff --git a/Engine/Config/BaseEngine.ini b/Engine/Config/BaseEngine.ini
--- a/Engine/Config/BaseEngine.ini
+++ b/Engine/Config/BaseEngine.ini
@@ -1184,7 +1184,7 @@ GameThreadHeartBeatStartSuspended=False
 GameThreadHeartBeatHitchDuration=-1.0
 GameThreadHeartBeatStackWalk=false
 IgnoreInlineBulkDataReloadEnsures=False
-UsePackageTrailer=False
+UsePackageTrailer=True
 HangDuration=0.0
 AssetLogShowsDiskPath=True
 AssetLogShowsAbsolutePath=True
