Updating [MaterialEditorSpawnNodes] section to use class path names

diff --git a/Engine/Config/BaseEditorPerProjectUserSettings.ini b/Engine/Config/BaseEditorPerProjectUserSettings.ini
--- a/Engine/Config/BaseEditorPerProjectUserSettings.ini
+++ b/Engine/Config/BaseEditorPerProjectUserSettings.ini
@@ -813,26 +813,26 @@ DefaultBrushRadius=128
 +Node=(TargetClass=/Script/MovieSceneCapture.UserDefinedCaptureProtocol TargetEvent="OnFinalize")
 
 [MaterialEditorSpawnNodes]
-+Node=(Class=MaterialExpressionAdd Key=A Shift=false Ctrl=false Alt=false)
-+Node=(Class=MaterialExpressionBumpOffset Key=B Shift=false Ctrl=false Alt=false)
-+Node=(Class=MaterialExpressionDivide Key=D Shift=false Ctrl=false Alt=false)
-+Node=(Class=MaterialExpressionPower Key=E Shift=false Ctrl=false Alt=false)
-+Node=(Class=MaterialExpressionMaterialFunctionCall Key=F Shift=false Ctrl=false Alt=false)
-+Node=(Class=MaterialExpressionIf Key=I Shift=false Ctrl=false Alt=false)
-+Node=(Class=MaterialExpressionLinearInterpolate Key=L Shift=false Ctrl=false Alt=false)
-+Node=(Class=MaterialExpressionMultiply Key=M Shift=false Ctrl=false Alt=false)
-+Node=(Class=MaterialExpressionNormalize Key=N Shift=false Ctrl=false Alt=false)
-+Node=(Class=MaterialExpressionOneMinus Key=O Shift=false Ctrl=false Alt=false)
-+Node=(Class=MaterialExpressionPanner Key=P Shift=false Ctrl=false Alt=false)
-+Node=(Class=MaterialExpressionReflectionVectorWS Key=R Shift=false Ctrl=false Alt=false)
-+Node=(Class=MaterialExpressionScalarParameter Key=S Shift=false Ctrl=false Alt=false)
-+Node=(Class=MaterialExpressionTextureSample Key=T Shift=false Ctrl=false Alt=false)
-+Node=(Class=MaterialExpressionTextureCoordinate Key=U Shift=false Ctrl=false Alt=false)
-+Node=(Class=MaterialExpressionVectorParameter Key=V Shift=false Ctrl=false Alt=false)
-+Node=(Class=MaterialExpressionConstant Key=One Shift=false Ctrl=false Alt=false)
-+Node=(Class=MaterialExpressionConstant2Vector Key=Two Shift=false Ctrl=false Alt=false)
-+Node=(Class=MaterialExpressionConstant3Vector Key=Three Shift=false Ctrl=false Alt=false)
-+Node=(Class=MaterialExpressionConstant4Vector Key=Four Shift=false Ctrl=false Alt=false)
++Node=(Class=/Script/Engine.MaterialExpressionAdd Key=A Shift=false Ctrl=false Alt=false)
++Node=(Class=/Script/Engine.MaterialExpressionBumpOffset Key=B Shift=false Ctrl=false Alt=false)
++Node=(Class=/Script/Engine.MaterialExpressionDivide Key=D Shift=false Ctrl=false Alt=false)
++Node=(Class=/Script/Engine.MaterialExpressionPower Key=E Shift=false Ctrl=false Alt=false)
++Node=(Class=/Script/Engine.MaterialExpressionMaterialFunctionCall Key=F Shift=false Ctrl=false Alt=false)
++Node=(Class=/Script/Engine.MaterialExpressionIf Key=I Shift=false Ctrl=false Alt=false)
++Node=(Class=/Script/Engine.MaterialExpressionLinearInterpolate Key=L Shift=false Ctrl=false Alt=false)
++Node=(Class=/Script/Engine.MaterialExpressionMultiply Key=M Shift=false Ctrl=false Alt=false)
++Node=(Class=/Script/Engine.MaterialExpressionNormalize Key=N Shift=false Ctrl=false Alt=false)
++Node=(Class=/Script/Engine.MaterialExpressionOneMinus Key=O Shift=false Ctrl=false Alt=false)
++Node=(Class=/Script/Engine.MaterialExpressionPanner Key=P Shift=false Ctrl=false Alt=false)
++Node=(Class=/Script/Engine.MaterialExpressionReflectionVectorWS Key=R Shift=false Ctrl=false Alt=false)
++Node=(Class=/Script/Engine.MaterialExpressionScalarParameter Key=S Shift=false Ctrl=false Alt=false)
++Node=(Class=/Script/Engine.MaterialExpressionTextureSample Key=T Shift=false Ctrl=false Alt=false)
++Node=(Class=/Script/Engine.MaterialExpressionTextureCoordinate Key=U Shift=false Ctrl=false Alt=false)
++Node=(Class=/Script/Engine.MaterialExpressionVectorParameter Key=V Shift=false Ctrl=false Alt=false)
++Node=(Class=/Script/Engine.MaterialExpressionConstant Key=One Shift=false Ctrl=false Alt=false)
++Node=(Class=/Script/Engine.MaterialExpressionConstant2Vector Key=Two Shift=false Ctrl=false Alt=false)
++Node=(Class=/Script/Engine.MaterialExpressionConstant3Vector Key=Three Shift=false Ctrl=false Alt=false)
++Node=(Class=/Script/Engine.MaterialExpressionConstant4Vector Key=Four Shift=false Ctrl=false Alt=false)
 
 [WidgetTemplatesExpanded]
 Common=True
