TimeOfDay: Add default tracks for TimeOfDayActor.
Add support for default component transform tracks for TimeOfDaySequences via RelativeLocation, RelativeRotation or RelativeScale3D default property tracks.

diff --git a/Engine/Config/BaseEditorPerProjectUserSettings.ini b/Engine/Config/BaseEditorPerProjectUserSettings.ini
--- a/Engine/Config/BaseEditorPerProjectUserSettings.ini
+++ b/Engine/Config/BaseEditorPerProjectUserSettings.ini
@@ -909,6 +909,7 @@ FrameNumberDisplayFormat=Frames
 +TrackSettings=(MatchingActorClass=/Script/Engine.CameraActor,DefaultTracks=(/Script/MovieSceneTracks.MovieScene3DTransformTrack),DefaultPropertyTracks=((ComponentPath="CameraComponent",PropertyPath="FieldOfView")))
 +TrackSettings=(MatchingActorClass=/Script/CinematicCamera.CineCameraActor,DefaultPropertyTracks=((ComponentPath="CameraComponent",PropertyPath="CurrentFocalLength"),(ComponentPath="CameraComponent",PropertyPath="FocusSettings.ManualFocusDistance"),(ComponentPath="CameraComponent",PropertyPath="CurrentAperture")),ExcludeDefaultPropertyTracks=((ComponentPath="CameraComponent",PropertyPath="FieldOfView")))
 +TrackSettings=(MatchingActorClass=/Script/Engine.Light,DefaultTracks=(None),DefaultPropertyTracks=((ComponentPath="LightComponent0",PropertyPath="Intensity"),(ComponentPath="LightComponent0",PropertyPath="LightColor")))
++TrackSettings=(MatchingActorClass=/Script/TimeOfDay.TimeOfDayActor,DefaultPropertyTracks=((ComponentPath="Sun",PropertyPath="RelativeRotation"),(ComponentPath="Sun",PropertyPath="LightColor"),(ComponentPath="Sun",PropertyPath="Intensity"),(ComponentPath="SkyLight",PropertyPath="LightColor"),(ComponentPath="SkyLight",PropertyPath="Intensity"))
 
 [/Script/MovieSceneTools.MovieSceneToolsProjectSettings]
 +FbxSettings=(FbxPropertyName="FieldOfView", PropertyPath=(ComponentName="CameraComponent",PropertyName="FieldOfView"))
