- Enabled x86 and arm64 configs in Xcode by default for all targets

diff --git a/Engine/Config/BaseEngine.ini b/Engine/Config/BaseEngine.ini
--- a/Engine/Config/BaseEngine.ini
+++ b/Engine/Config/BaseEngine.ini
@@ -2764,8 +2764,8 @@ MetalLanguageVersion=5
 +TargetedRHIs=SF_METAL_SM5
 UseFastIntrinsics=False
 EnableMathOptimisations=True
-TargetArchitecture=MacTargetArchitectureIntel
-EditorTargetArchitecture=MacTargetArchitectureIntel
+TargetArchitecture=MacTargetArchitectureUniversal
+EditorTargetArchitecture=MacTargetArchitectureUniversal
 
 [HMDPluginPriority]
 ; Since SteamVR also works with the Oculus Rift, give priority to the native Oculus plugin before trying SteamVR.
