Enabling 'AttemptToFindShortTypeNamesInMetaData' engine and editor startup tests.

diff --git a/Engine/Config/BaseEngine.ini b/Engine/Config/BaseEngine.ini
--- a/Engine/Config/BaseEngine.ini
+++ b/Engine/Config/BaseEngine.ini
@@ -103,6 +103,22 @@ EngineModuleReflectedUninitializedPropertyVerbosity=Error
 ProjectModuleReflectedUninitializedPropertyVerbosity=Error
 ObjectReferenceReflectedUninitializedPropertyVerbosity=Error
 
+[CoreUObject.ShortTypeNameInMetaDataCheck]
++MetaDataKeys=AllowedClasses
++MetaDataKeys=DisallowedClasses
++MetaDataKeys=BaseStruct
++MetaDataKeys=BaseClass
++MetaDataKeys=HasNativeMake
++MetaDataKeys=HasnativeBreak
++MetaDataKeys=RowType
++MetaDataKeys=BitmaskEnum
++MetaDataKeys=ActorClassToPoll
++MetaDataKeys=ComponentClass
++MetaDataKeys=MustImplement
++MetaDataKeys=MetaClass
++MetaDataKeys=MetaStruct
++MetaDataKeys=RequiredAssetDataTags
+
 [/Script/Engine.Engine]
 ConsoleClassName=/Script/Engine.Console
 GameViewportClientClassName=/Script/Engine.GameViewportClient
@@ -3028,7 +3044,6 @@ PreallocatedBufferCount=2
 [Plugins]
 ; Do not allow plugins to override EditorPerProjectUserSettings because it needs allow save
 +ConfigFilesPluginsCannotOverride=EditorPerProjectUserSettings
-
 [ClothSettings]
 ;  Extra padding added to the bounds so that the collision can still be detected after being resolved
 EnvironmentCollisionPadding = 2.f
diff --git a/Engine/Plugins/Runtime/GameplayAbilities/Source/GameplayAbilities/Public/GameplayCueSet.h b/Engine/Plugins/Runtime/GameplayAbilities/Source/GameplayAbilities/Public/GameplayCueSet.h
--- a/Engine/Plugins/Runtime/GameplayAbilities/Source/GameplayAbilities/Public/GameplayCueSet.h
+++ b/Engine/Plugins/Runtime/GameplayAbilities/Source/GameplayAbilities/Public/GameplayCueSet.h
@@ -24,7 +24,7 @@ struct FGameplayCueNotifyData
 	UPROPERTY(EditAnywhere, Category=GameplayCue)
 	FGameplayTag GameplayCueTag;
 
-	UPROPERTY(EditAnywhere, Category=GameplayCue, meta=(AllowedClasses="GameplayCueNotify"))
+	UPROPERTY(EditAnywhere, Category=GameplayCue, meta=(AllowedClasses="/Script/GameplayAbilities.GameplayCueNotify_Static, /Script/GameplayAbilities.GameplayCueNotify_Actor"))
 	FSoftObjectPath GameplayCueNotifyObj;
 
 	UPROPERTY(transient)
diff --git a/Engine/Plugins/Runtime/Oculus/OculusVR/Source/OculusHMD/Public/OculusHMDTypes.h b/Engine/Plugins/Runtime/Oculus/OculusVR/Source/OculusHMD/Public/OculusHMDTypes.h
--- a/Engine/Plugins/Runtime/Oculus/OculusVR/Source/OculusHMD/Public/OculusHMDTypes.h
+++ b/Engine/Plugins/Runtime/Oculus/OculusVR/Source/OculusHMD/Public/OculusHMDTypes.h
@@ -16,7 +16,7 @@ struct FOculusSplashDesc
 	GENERATED_USTRUCT_BODY()
 
 	UPROPERTY(config, EditAnywhere, Category = Settings, meta = (
-		AllowedClasses = "Texture",
+		AllowedClasses = "/Script/Engine.Texture",
 		ToolTip = "Texture to display"))
 	FSoftObjectPath		TexturePath;
 
