INI changes for use of desired structured namespaces when using Zen as a DDC backend.

diff --git a/Engine/Config/BaseEngine.ini b/Engine/Config/BaseEngine.ini
--- a/Engine/Config/BaseEngine.ini
+++ b/Engine/Config/BaseEngine.ini
@@ -2245,13 +2245,13 @@ EnterprisePak=(Type=ReadPak, Filename=../../../Enterprise/DerivedDataCache/Compr
 [ZenOnlyDDC]
 Root=(Type=KeyLength, Length=120, Inner=AsyncPut)
 AsyncPut=(Type=AsyncPut, Inner=Zen)
-Zen=(Type=Zen, Namespace="ue4.ddc", StructuredNamespace="ue4.ddc")
+Zen=(Type=Zen, Namespace="ue4.ddc", StructuredNamespace="ue.ddc")
 
 [ZenDDC]
 Root=(Type=KeyLength, Length=120, Inner=AsyncPut)
 AsyncPut=(Type=AsyncPut, Inner=Hierarchy)
 Hierarchy=(Type=Hierarchical, Inner=Pak, Inner=EnginePak, Inner=Zen, Inner=Shared)
-Zen=(Type=Zen, Namespace="ue4.ddc", StructuredNamespace="ue4.ddc")
+Zen=(Type=Zen, Namespace="ue4.ddc", StructuredNamespace="ue.ddc")
 Shared=(Type=FileSystem, UnusedFileAge=10, FoldersToClean=10, ConsiderSlowAt=70, Path=?EpicDDC, EnvPathOverride=UE-SharedDataCachePath, EditorOverrideSetting=SharedDerivedDataCache, CommandLineOverride=SharedDataCachePath)
 Pak=(Type=ReadPak, Filename="%GAMEDIR%DerivedDataCache/DDC.ddp")
 EnginePak=(Type=ReadPak, Filename="%ENGINEDIR%DerivedDataCache/DDC.ddp")
