Add r.GPUStatsEnabled=0 to ConsoleVariable.ini to solve GPU related crash at startup

diff --git a/Engine/Config/ConsoleVariables.ini b/Engine/Config/ConsoleVariables.ini
--- a/Engine/Config/ConsoleVariables.ini
+++ b/Engine/Config/ConsoleVariables.ini
@@ -131,3 +131,5 @@ fx.SkipVectorVMBackendOptimizations=1
 ; Uncomment to force to import embedded tessellation instead of meshing BRep model. Default is false
 ;ds.CADTranslator.PreferJtFileEmbeddedTessellation=true
 
+; work-around to fix GPU crash at start-up
+r.GPUStatsEnabled=0
