Undo CL 20641952 since this is not something we want propagated across all UE projects.

#rnx

diff --git a/Engine/Config/ConsoleVariables.ini b/Engine/Config/ConsoleVariables.ini
--- a/Engine/Config/ConsoleVariables.ini
+++ b/Engine/Config/ConsoleVariables.ini
@@ -130,6 +130,3 @@ fx.SkipVectorVMBackendOptimizations=1
 
 ; Uncomment to force to import embedded tessellation instead of meshing BRep model. Default is false
 ;ds.CADTranslator.PreferJtFileEmbeddedTessellation=true
-
-; work-around to fix GPU crash at start-up
-r.GPUStatsEnabled=0
