Fix short path name warnings in macro graph hotkey setup during BP editor startup.

Change summary:
- Updated macro command hotkey config records in base EditorPerProjectUserSettings.ini file to specify the full path for macro graphs.
- Modified FBlueprintSpawnNodeCommands::RegisterCommands() to directly route to macro graph objects rather than doing a search over the full set of loaded BPs.

diff --git a/Engine/Config/BaseEditorPerProjectUserSettings.ini b/Engine/Config/BaseEditorPerProjectUserSettings.ini
--- a/Engine/Config/BaseEditorPerProjectUserSettings.ini
+++ b/Engine/Config/BaseEditorPerProjectUserSettings.ini
@@ -775,14 +775,14 @@ DefaultBrushRadius=128
 [BlueprintSpawnNodes]
 ; Comment box is bound to C, but that is handled differently due to it needing to work without clicking
 +Node=(Class=/Script/Engine.Actor:ReceiveBeginPlay Key=P Shift=false Ctrl=false Alt=false)
-+Node=(Class="Do N" Key=N Shift=false Ctrl=false Alt=false)
++Node=(Class="/Engine/EditorBlueprintResources/StandardMacros.StandardMacros:Do N" Key=N Shift=false Ctrl=false Alt=false)
 +Node=(Class=/Script/Engine.KismetSystemLibrary:Delay Key=D Shift=false Ctrl=false Alt=false)
 +Node=(Class=/Script/BlueprintGraph.K2Node_IfThenElse Key=B Shift=false Ctrl=false Alt=false)
 +Node=(Class=/Script/BlueprintGraph.K2Node_ExecutionSequence Key=S Shift=false Ctrl=false Alt=false)
-+Node=(Class=Gate Key=G Shift=false Ctrl=false Alt=false)
++Node=(Class=/Engine/EditorBlueprintResources/StandardMacros.StandardMacros:Gate Key=G Shift=false Ctrl=false Alt=false)
 +Node=(Class=/Script/BlueprintGraph.K2Node_MultiGate Key=M Shift=false Ctrl=false Alt=false)
-+Node=(Class=ForEachLoop Key=F Shift=false Ctrl=false Alt=false)
-+Node=(Class=DoOnce Key=O Shift=false Ctrl=false Alt=false)
++Node=(Class=/Engine/EditorBlueprintResources/StandardMacros.StandardMacros:ForEachLoop Key=F Shift=false Ctrl=false Alt=false)
++Node=(Class=/Engine/EditorBlueprintResources/StandardMacros.StandardMacros:DoOnce Key=O Shift=false Ctrl=false Alt=false)
 +Node=(Class=/Script/Engine.KismetArrayLibrary:Array_Get Key=A Shift=false Ctrl=false Alt=false)
 
 [DefaultEventNodes]
diff --git a/Engine/Source/Editor/Kismet/Private/BlueprintEditorCommands.cpp b/Engine/Source/Editor/Kismet/Private/BlueprintEditorCommands.cpp
--- a/Engine/Source/Editor/Kismet/Private/BlueprintEditorCommands.cpp
+++ b/Engine/Source/Editor/Kismet/Private/BlueprintEditorCommands.cpp
@@ -339,27 +339,16 @@ void FBlueprintSpawnNodeCommands::RegisterCommands()
 
 			CommandLabel = FoundFunction->GetName();
 		}
-		else
+		else if(UEdGraph* FoundGraph = FindFirstObject<UEdGraph>(*ClassName, EFindFirstObjectOptions::ExactClass, ELogVerbosity::Warning, TEXT("looking for SpawnNodes (macro)")))
 		{
-			// Check for a macro graph that matches the passed in class name
-			for (TObjectIterator<UBlueprint> BlueprintIt; BlueprintIt; ++BlueprintIt)
+			// If the name matches that of a macro graph, set up a spawn info that can generate macro graph actions
+			if (const UBlueprint* Blueprint = FBlueprintEditorUtils::FindBlueprintForGraph(FoundGraph))
 			{
-				UBlueprint* MacroBP = *BlueprintIt;
-				if(MacroBP->BlueprintType == BPTYPE_MacroLibrary)
+				if (Blueprint->BlueprintType == BPTYPE_MacroLibrary)
 				{
-					// getting 'top-level' of the macros
-					for (TArray<UEdGraph*>::TIterator GraphIt(MacroBP->MacroGraphs); GraphIt; ++GraphIt)
-					{
-						UEdGraph* MacroGraph = *GraphIt;
-						// The class name matches that of a macro, so setup a spawn info that can generate macro graph actions
-						if(MacroGraph->GetName() == ClassName)
-						{
-							CommandLabel = MacroGraph->GetName();
+					CommandLabel = FoundGraph->GetName();
 
-							InfoPtr = MakeShareable( new FMacroNodeSpawnInfo(MacroGraph));
-						}
-
-					}
+					InfoPtr = MakeShareable(new FMacroNodeSpawnInfo(FoundGraph));
 				}
 			}
 		}
