Lower the minimum required nvidia driver version for DX11, since some older GPUs that we still support don't receive updates anymore.

473.47 is the latest driver available for GeForce 500 series and such.

Also, fix bad GPU driver warning added to the log incorrectly due to bad else block nesting.

#rnx

diff --git a/Engine/Config/BaseHardware.ini b/Engine/Config/BaseHardware.ini
--- a/Engine/Config/BaseHardware.ini
+++ b/Engine/Config/BaseHardware.ini
@@ -39,7 +39,9 @@
 SuggestedDriverVersion="516.40"
 
 ; To prevent problems with older drivers.
-+DriverDenyList=(DriverVersion="<516.01", Reason="These driver versions have known stability issues and missing features")
++DriverDenyList=(DriverVersion="<473.47", Reason="These driver versions have known stability issues and missing features", RHI="D3D11")
++DriverDenyList=(DriverVersion="<516.01", Reason="These driver versions have known stability issues and missing features", RHI="D3D12")
++DriverDenyList=(DriverVersion="<516.01", Reason="These driver versions have known stability issues and missing features", RHI="Vulkan")
 
 ;-----------------------------------------------------------------------------------------------------------------
 
