Change the suggested nvidia driver version to 516.25, since that's the earliest driver version which contains the fix and is also available for pro GPUs.

#rnx

diff --git a/Engine/Config/BaseHardware.ini b/Engine/Config/BaseHardware.ini
--- a/Engine/Config/BaseHardware.ini
+++ b/Engine/Config/BaseHardware.ini
@@ -36,7 +36,7 @@
 
 
 [GPU_NVIDIA Windows]
-SuggestedDriverVersion="516.40"
+SuggestedDriverVersion="516.25"
 
 ; To prevent problems with older drivers.
 +DriverDenyList=(DriverVersion="<473.47", Reason="These driver versions have known stability issues and missing features", RHI="D3D11")
