Enable automatic editor optional sidecar data for Niagara asset types

diff --git a/Engine/Config/BaseEditor.ini b/Engine/Config/BaseEditor.ini
--- a/Engine/Config/BaseEditor.ini
+++ b/Engine/Config/BaseEditor.ini
@@ -366,6 +366,11 @@ CookWorkerCount=1
 +CookOnTheFlyConfigSettingDenyList=*.Engine:/Script/UnrealEd.CookerSettings:bIterativeCookingForFileCookContent
 +CookOnTheFlyConfigSettingDenyList=*.Engine:OnlineSubsystemMcp
 
+; Asset type to automatically generate editor optional sidecar data for
++AutomaticOptionalInclusionAssetType=/Script/Niagara.NiagaraScript
++AutomaticOptionalInclusionAssetType=/Script/Niagara.NiagaraEmitter
++AutomaticOptionalInclusionAssetType=/Script/Niagara.NiagaraSystem
+
 [EditorDomain]
 EditorDomainEnabled=false
 SaveUnversioned=true
