Disable automatic optional for Niagara type during bulk data refactor

diff --git a/Engine/Config/BaseEditor.ini b/Engine/Config/BaseEditor.ini
--- a/Engine/Config/BaseEditor.ini
+++ b/Engine/Config/BaseEditor.ini
@@ -367,9 +367,9 @@ CookWorkerCount=1
 +CookOnTheFlyConfigSettingDenyList=*.Engine:OnlineSubsystemMcp
 
 ; Asset type to automatically generate editor optional sidecar data for
-+AutomaticOptionalInclusionAssetType=/Script/Niagara.NiagaraScript
-+AutomaticOptionalInclusionAssetType=/Script/Niagara.NiagaraEmitter
-+AutomaticOptionalInclusionAssetType=/Script/Niagara.NiagaraSystem
+;+AutomaticOptionalInclusionAssetType=/Script/Niagara.NiagaraScript
+;+AutomaticOptionalInclusionAssetType=/Script/Niagara.NiagaraEmitter
+;+AutomaticOptionalInclusionAssetType=/Script/Niagara.NiagaraSystem
 
 [EditorDomain]
 EditorDomainEnabled=false
