Fix incorrect warning about outdated drivers on Intel GPUs caused by the fact that the driver version is reported in a different format for that vendor.

#rnx

diff --git a/Engine/Config/BaseHardware.ini b/Engine/Config/BaseHardware.ini
--- a/Engine/Config/BaseHardware.ini
+++ b/Engine/Config/BaseHardware.ini
@@ -79,7 +79,7 @@ SuggestedDriverVersion="21.2.6"
 
 [GPU_Intel Windows]
 SuggestedDriverVersion="30.0.101.1340"
-+DriverDenyList=(DriverVersion="<30.0.101.1191", Reason="These driver versions have known stability issues and missing features")
++DriverDenyList=(DriverVersion="<101.1191", Reason="These driver versions have known stability issues and missing features")
 
 ;-----------------------------------------------------------------------------------------------------------------
 
