TimeOfDay: Rename to DaySequence.

diff --git a/Engine/Config/BaseEditorPerProjectUserSettings.ini b/Engine/Config/BaseEditorPerProjectUserSettings.ini
--- a/Engine/Config/BaseEditorPerProjectUserSettings.ini
+++ b/Engine/Config/BaseEditorPerProjectUserSettings.ini
@@ -882,7 +882,7 @@ LoopMode=SLM_Loop
 bCleanPlaybackMode=false
 bShowLayerBars=false
 
-[TimeOfDaySequenceEditor SequencerSettings]
+[DaySequenceEditor SequencerSettings]
 bShowRangeSlider=true
 bKeepPlayRangeInSectionBounds=false
 ZeroPadFrames=4
@@ -904,13 +904,13 @@ FrameNumberDisplayFormat=Frames
 +TrackSettings=(MatchingActorClass=/Script/CinematicCamera.CineCameraActor,DefaultPropertyTracks=((ComponentPath="CameraComponent",PropertyPath="CurrentFocalLength"),(ComponentPath="CameraComponent",PropertyPath="FocusSettings.ManualFocusDistance"),(ComponentPath="CameraComponent",PropertyPath="CurrentAperture")),ExcludeDefaultPropertyTracks=((ComponentPath="CameraComponent",PropertyPath="FieldOfView")))
 +TrackSettings=(MatchingActorClass=/Script/Engine.Light,DefaultTracks=(None),DefaultPropertyTracks=((ComponentPath="LightComponent0",PropertyPath="Intensity"),(ComponentPath="LightComponent0",PropertyPath="LightColor")))
 
-[/Script/TimeOfDayEditor.TimeOfDayEditorSettings]
+[/Script/DaySequenceEditor.DaySequenceEditorSettings]
 +TrackSettings=(MatchingActorClass=/Script/Engine.StaticMeshActor,DefaultTracks=(/Script/MovieSceneTracks.MovieScene3DTransformTrack))
 +TrackSettings=(MatchingActorClass=/Script/Engine.SkeletalMeshActor,DefaultTracks=(/Script/MovieSceneTracks.MovieScene3DTransformTrack,/Script/MovieSceneTracks.MovieSceneSkeletalAnimationTrack))
 +TrackSettings=(MatchingActorClass=/Script/Engine.CameraActor,DefaultTracks=(/Script/MovieSceneTracks.MovieScene3DTransformTrack),DefaultPropertyTracks=((ComponentPath="CameraComponent",PropertyPath="FieldOfView")))
 +TrackSettings=(MatchingActorClass=/Script/CinematicCamera.CineCameraActor,DefaultPropertyTracks=((ComponentPath="CameraComponent",PropertyPath="CurrentFocalLength"),(ComponentPath="CameraComponent",PropertyPath="FocusSettings.ManualFocusDistance"),(ComponentPath="CameraComponent",PropertyPath="CurrentAperture")),ExcludeDefaultPropertyTracks=((ComponentPath="CameraComponent",PropertyPath="FieldOfView")))
 +TrackSettings=(MatchingActorClass=/Script/Engine.Light,DefaultTracks=(None),DefaultPropertyTracks=((ComponentPath="LightComponent0",PropertyPath="Intensity"),(ComponentPath="LightComponent0",PropertyPath="LightColor")))
-+TrackSettings=(MatchingActorClass=/Script/TimeOfDay.TimeOfDayActor,DefaultPropertyTracks=((ComponentPath="Sun",PropertyPath="RelativeRotation"),(ComponentPath="Sun",PropertyPath="LightColor"),(ComponentPath="Sun",PropertyPath="Intensity"),(ComponentPath="SkyLight",PropertyPath="LightColor"),(ComponentPath="SkyLight",PropertyPath="Intensity")))
++TrackSettings=(MatchingActorClass=/Script/DaySequence.DaySequenceActor,DefaultPropertyTracks=((ComponentPath="Sun",PropertyPath="RelativeRotation"),(ComponentPath="Sun",PropertyPath="LightColor"),(ComponentPath="Sun",PropertyPath="Intensity"),(ComponentPath="SkyLight",PropertyPath="LightColor"),(ComponentPath="SkyLight",PropertyPath="Intensity")))
 
 [/Script/MovieSceneTools.MovieSceneToolsProjectSettings]
 +FbxSettings=(FbxPropertyName="FieldOfView", PropertyPath=(ComponentName="CameraComponent",PropertyName="FieldOfView"))
