Fix New Level dialog icons for Open World and Basic

diff --git a/Engine/Config/BaseEngine.ini b/Engine/Config/BaseEngine.ini
--- a/Engine/Config/BaseEngine.ini
+++ b/Engine/Config/BaseEngine.ini
@@ -1606,8 +1606,8 @@ InEditorGameURLOptions=
 
 [/Script/UnrealEd.UnrealEdEngine]
 AutoSaveIndex=0
-+TemplateMapInfos=(Thumbnail="/Engine/Maps/Templates/Thumbnails/OpenWorld",Map="/Engine/Maps/Templates/OpenWorld",DisplayName="Open World",Category="OpenWorld")
-+TemplateMapInfos=(Thumbnail="/Engine/Maps/Templates/Thumbnails/Basic",Map="/Engine/Maps/Templates/Template_Default",DisplayName="Basic")
++TemplateMapInfos=(Thumbnail="/Engine/Maps/Templates/Thumbnails/OpenWorld.OpenWorld",Map="/Engine/Maps/Templates/OpenWorld.OpenWorld",DisplayName="Open World",Category="OpenWorld")
++TemplateMapInfos=(Thumbnail="/Engine/Maps/Templates/Thumbnails/Basic.Basic",Map="/Engine/Maps/Templates/Template_Default.Template_Default",DisplayName="Basic")
 
 +PackagesToBeFullyLoadedAtStartup=/Engine/EditorMaterials/Cross
 +PackagesToBeFullyLoadedAtStartup=/Engine/EditorMaterials/Cross_Mat
