Fix UUnrealEdEngine::IsTemplateMap() and adjust BaseEngine.ini template map paths back to how they were before (no extra dots)

diff --git a/Engine/Config/BaseEngine.ini b/Engine/Config/BaseEngine.ini
--- a/Engine/Config/BaseEngine.ini
+++ b/Engine/Config/BaseEngine.ini
@@ -1608,8 +1608,8 @@ InEditorGameURLOptions=
 
 [/Script/UnrealEd.UnrealEdEngine]
 AutoSaveIndex=0
-+TemplateMapInfos=(Thumbnail="/Engine/Maps/Templates/Thumbnails/OpenWorld.OpenWorld",Map="/Engine/Maps/Templates/OpenWorld.OpenWorld",DisplayName="Open World",Category="OpenWorld")
-+TemplateMapInfos=(Thumbnail="/Engine/Maps/Templates/Thumbnails/Basic.Basic",Map="/Engine/Maps/Templates/Template_Default.Template_Default",DisplayName="Basic")
++TemplateMapInfos=(Thumbnail="/Engine/Maps/Templates/Thumbnails/OpenWorld.OpenWorld",Map="/Engine/Maps/Templates/OpenWorld",DisplayName="Open World",Category="OpenWorld")
++TemplateMapInfos=(Thumbnail="/Engine/Maps/Templates/Thumbnails/Basic.Basic",Map="/Engine/Maps/Templates/Template_Default",DisplayName="Basic")
 
 +PackagesToBeFullyLoadedAtStartup=/Engine/EditorMaterials/Cross
 +PackagesToBeFullyLoadedAtStartup=/Engine/EditorMaterials/Cross_Mat
diff --git a/Engine/Source/Editor/UnrealEd/Private/UnrealEdSrv.cpp b/Engine/Source/Editor/UnrealEd/Private/UnrealEdSrv.cpp
--- a/Engine/Source/Editor/UnrealEd/Private/UnrealEdSrv.cpp
+++ b/Engine/Source/Editor/UnrealEd/Private/UnrealEdSrv.cpp
@@ -1316,7 +1316,7 @@ bool UUnrealEdEngine::IsTemplateMap( const FString& MapName ) const
 {
 	for (const FTemplateMapInfo& It : GetTemplateMapInfos())
 	{
-		if (It.Map == MapName)
+		if (It.Map.GetLongPackageName() == MapName)
 		{
 			return true;
 		}
