Include libc++_shared.so in Android packages by default due to EOSSDK dependency on it

diff --git a/Engine/Config/BaseEngine.ini b/Engine/Config/BaseEngine.ini
--- a/Engine/Config/BaseEngine.ini
+++ b/Engine/Config/BaseEngine.ini
@@ -2724,7 +2724,7 @@ bAllowPatchOBBFile=False
 bAllowOverflowOBBFiles=False
 bExtractNativeLibs=true
 bEnableDomStorage=False
-bSkipLibCpp=True
+bSkipLibCpp=False
 bDisableLibCppSharedDependencyValidation=True
 bDisableFunctionDataSections=False
 bEnableAdvancedBinaryCompression=False
