Disabled Octahedral Irradiance format for Screen Probe Gather on Epic GI scalability settings as it doesn't handle gradients well.

diff --git a/Engine/Config/BaseScalability.ini b/Engine/Config/BaseScalability.ini
--- a/Engine/Config/BaseScalability.ini
+++ b/Engine/Config/BaseScalability.ini
@@ -262,6 +262,7 @@ r.Lumen.ScreenProbeGather.RadianceCache.NumProbesToTraceBudget=300
 r.Lumen.ScreenProbeGather.ScreenSpaceBentNormal=0
 r.Lumen.ScreenProbeGather.DownsampleFactor=32
 r.Lumen.ScreenProbeGather.TracingOctahedronResolution=8
+r.Lumen.ScreenProbeGather.IrradianceFormat=1
 r.Lumen.ScreenProbeGather.StochasticInterpolation=1
 r.Lumen.ScreenProbeGather.FullResolutionJitterWidth=1
 r.Lumen.TranslucencyVolume.TraceFromVolume=0
@@ -279,6 +280,7 @@ r.Lumen.ScreenProbeGather.RadianceCache.NumProbesToTraceBudget=300
 r.Lumen.ScreenProbeGather.ScreenSpaceBentNormal=1
 r.Lumen.ScreenProbeGather.DownsampleFactor=16
 r.Lumen.ScreenProbeGather.TracingOctahedronResolution=8
+r.Lumen.ScreenProbeGather.IrradianceFormat=0
 r.Lumen.ScreenProbeGather.StochasticInterpolation=0
 r.Lumen.ScreenProbeGather.FullResolutionJitterWidth=1
 r.Lumen.TranslucencyVolume.TraceFromVolume=1
@@ -297,6 +299,7 @@ r.Lumen.ScreenProbeGather.RadianceCache.NumProbesToTraceBudget=100000
 r.Lumen.ScreenProbeGather.ScreenSpaceBentNormal=1
 r.Lumen.ScreenProbeGather.DownsampleFactor=8
 r.Lumen.ScreenProbeGather.TracingOctahedronResolution=16
+r.Lumen.ScreenProbeGather.IrradianceFormat=0
 r.Lumen.ScreenProbeGather.StochasticInterpolation=0
 r.Lumen.ScreenProbeGather.FullResolutionJitterWidth=.5
 r.Lumen.TranslucencyVolume.TraceFromVolume=1
