Fix default CVar value in ConsoleVaraibles.ini

diff --git a/Engine/Config/ConsoleVariables.ini b/Engine/Config/ConsoleVariables.ini
--- a/Engine/Config/ConsoleVariables.ini
+++ b/Engine/Config/ConsoleVariables.ini
@@ -126,7 +126,7 @@ fx.SkipVectorVMBackendOptimizations=1
 ;ds.CADTranslator.EnableTimeControl=1
 
 ; Uncomment to enable/disable CADKernel tessellation. If disable, CAD SDK tessellator is used.
-ds.CADTranslator.DisableCADKernelTessellation=false
+;ds.CADTranslator.DisableCADKernelTessellation=false
 
 ; Perform a welding of the mesh to try to stitch mesh cracks.
 ;ds.CADTranslator.SewMeshIfNeeded=false
