Add CVar description in ConsoleVariables.ini

diff --git a/Engine/Config/ConsoleVariables.ini b/Engine/Config/ConsoleVariables.ini
--- a/Engine/Config/ConsoleVariables.ini
+++ b/Engine/Config/ConsoleVariables.ini
@@ -128,8 +128,47 @@ fx.SkipVectorVMBackendOptimizations=1
 ; Uncomment to enable/disable CADKernel tessellation. If disable, CAD SDK tessellator is used.
 ;ds.CADTranslator.DisableCADKernelTessellation=false
 
+; Set Welding threshold for Heal/Sew stitching methods (in cm)
+; Default value of StitchingTolerance is 0.001 cm
+;ds.CADTranslator.StitchingTolerance=0.001
+
 ; Perform a welding of the mesh to try to stitch mesh cracks.
-;ds.CADTranslator.SewMeshIfNeeded=false
+; Default value is true
+;ds.CADTranslator.SewMeshIfNeeded=true
+
+; Scale factor to change the unit of the DMU (Only applies to TechSoft import.)
+; Default value of UnitScale is 1 i.e. unit = cm
+;ds.CADTranslator.UnitScale=1
+
+; Perform a welding of the mesh to try to stitch mesh cracks
+; This welding is performed respecting the ds.CADTranslator.StitchingTolerance
+; Default is enable
+;ds.CADTranslator.SewMeshIfNeeded=true
+
+; Avoid duplicated triangles in a mesh (i.e. two triangle with the same triplet of vertices)
+; This problem appears when a surface patch is duplicated. Due to the two identical mesh of the patches, the global mesh has duplicated triangles
+; The time cost of this process is about 20% for 10e6 triangles with a complexity in nlog(n) due to the use of a map
+; As this problem is rare, and to avoid the cost of the process, this option is disable by default. 
+; To fix StaticMesh with duplicated triangles, the StaticMesh can be remesh with this option.
+; Default value is false.
+;ds.CADTranslator.RemoveDuplicatedTriangle=false
+
+; Factor to allow to use smaller value than the defined minimal value of metric meshing parameters (i.e. Chord error > 0.005 cm, Max Edge Length > 1. cm)
+; The used value of the meshing parameter is value * MeshingParameterFactor
+; Default value is 1.
+;ds.CADTranslator.MeshingParameterFactor=1
+
+; Alias Import Option
+; Enable Sew action merges BReps according to their material i.e. only BReps associated with same material can be merged together.
+; Default is disable
+;ds.CADTranslator.Alias.SewByColor=false
+
+; Alias Import Option
+; If true, the first level of actors in the outliner is the layers.
+; Default is false
+;ds.CADTranslator.Alias.LayersAsActors=false
 
-; Uncomment to force to import embedded tessellation instead of meshing BRep model. Default is false
-;ds.CADTranslator.PreferJtFileEmbeddedTessellation=true
+; JT Import Option
+; Force to import embedded tessellation instead of meshing BRep model.
+; Default is false
+;ds.CADTranslator.PreferJtFileEmbeddedTessellation=false
\ No newline at end of file
