Interchange: Add missing default value for the GenericPipelineClass
#rnx

diff --git a/Engine/Config/BaseEngine.ini b/Engine/Config/BaseEngine.ini
--- a/Engine/Config/BaseEngine.ini
+++ b/Engine/Config/BaseEngine.ini
@@ -3075,3 +3075,4 @@ FilePickerClass=/Script/InterchangeEditorUtilities.InterchangeFilePickerGeneric
 bStaticMeshUseSmoothEdgesIfSmoothingInformationIsMissing=True
 ContentImportSettings=(DefaultPipelineStackOverride=((Materials, "Materials"),(Textures, "Textures")),ShowPipelineStacksConfigurationDialogOverride=((Textures, False)),PipelineStacks=(("Assets", (Pipelines=("/Interchange/Pipelines/DefaultAssetsPipeline.DefaultAssetsPipeline"))),("Materials", (Pipelines=("/Interchange/Pipelines/DefaultMaterialPipeline.DefaultMaterialPipeline"))),("Textures", (Pipelines=("/Interchange/Pipelines/DefaultTexturePipeline.DefaultTexturePipeline")))),DefaultPipelineStack="Assets",PipelineConfigurationDialogClass="/Script/InterchangeEditorPipelines.InterchangePipelineConfigurationGeneric",bShowPipelineStacksConfigurationDialog=True)
 SceneImportSettings=(PipelineStacks=(("Scene", (Pipelines=("/Interchange/Pipelines/DefaultSceneAssetsPipeline.DefaultSceneAssetsPipeline","/Interchange/Pipelines/DefaultSceneLevelPipeline.DefaultSceneLevelPipeline")))),DefaultPipelineStack="Scene",PipelineConfigurationDialogClass="/Script/InterchangeEditorPipelines.InterchangePipelineConfigurationGeneric",bShowPipelineStacksConfigurationDialog=True)
+GenericPipelineClass=/Script/InterchangePipelines.InterchangeGenericAssetsPipeline
\ No newline at end of file
