Remove duplicated CVar

diff --git a/Engine/Config/ConsoleVariables.ini b/Engine/Config/ConsoleVariables.ini
--- a/Engine/Config/ConsoleVariables.ini
+++ b/Engine/Config/ConsoleVariables.ini
@@ -132,7 +132,8 @@ fx.SkipVectorVMBackendOptimizations=1
 ; Default value of StitchingTolerance is 0.001 cm
 ;ds.CADTranslator.StitchingTolerance=0.001
 
-; Perform a welding of the mesh to try to stitch mesh cracks.
+; Perform a welding of the mesh to try to stitch mesh cracks
+; This welding is performed respecting the ds.CADTranslator.StitchingTolerance
 ; Default value is true
 ;ds.CADTranslator.SewMeshIfNeeded=true
 
@@ -140,11 +141,6 @@ fx.SkipVectorVMBackendOptimizations=1
 ; Default value of UnitScale is 1 i.e. unit = cm
 ;ds.CADTranslator.UnitScale=1
 
-; Perform a welding of the mesh to try to stitch mesh cracks
-; This welding is performed respecting the ds.CADTranslator.StitchingTolerance
-; Default is enable
-;ds.CADTranslator.SewMeshIfNeeded=true
-
 ; Avoid duplicated triangles in a mesh (i.e. two triangle with the same triplet of vertices)
 ; This problem appears when a surface patch is duplicated. Due to the two identical mesh of the patches, the global mesh has duplicated triangles
 ; The time cost of this process is about 20% for 10e6 triangles with a complexity in nlog(n) due to the use of a map
