Change minimum metal version in base engine

diff --git a/Engine/Config/BaseEngine.ini b/Engine/Config/BaseEngine.ini
--- a/Engine/Config/BaseEngine.ini
+++ b/Engine/Config/BaseEngine.ini
@@ -2658,7 +2658,7 @@ bAutomaticSigning=false
 UseFastIntrinsics=False
 ForceFloats=True
 EnableMathOptimisations=True
-MetalLanguageVersion=0
+MetalLanguageVersion=7
 bDisableMotionData=False
 bEnableAdvertisingIdentifier=True
 
@@ -2802,7 +2802,7 @@ bTarget32Bit=false
 bEnableRayTracing=true
 
 [/Script/MacTargetPlatform.MacTargetSettings]
-MetalLanguageVersion=0
+MetalLanguageVersion=7
 +TargetedRHIs=SF_METAL_SM5
 UseFastIntrinsics=False
 EnableMathOptimisations=True
