Android: By default acquire next VK image on frame start

diff --git a/Engine/Config/BaseDeviceProfiles.ini b/Engine/Config/BaseDeviceProfiles.ini
--- a/Engine/Config/BaseDeviceProfiles.ini
+++ b/Engine/Config/BaseDeviceProfiles.ini
@@ -617,7 +617,8 @@ bIsVisibleForAssets=True
 
 +CVars=r.MobileContentScaleFactor=1.0
 +CVars=slate.AbsoluteIndices=1
-+CVars=r.Vulkan.DelayAcquireBackBuffer=2
+; acquire next image on frame start
++CVars=r.Vulkan.DelayAcquireBackBuffer=0
 +CVars=r.Vulkan.RobustBufferAccess=1
 +CVars=r.Vulkan.DescriptorSetLayoutMode=2
 ; Don't enable Vulkan by default. Specific device profiles can set this cvar to 0 to enable Vulkan.
