PR #9564: Fixed typo in "Set Fade In" node of DecalComponent (Contributed by [REDACTED_PERSON])

Also adds core redirect for backward compatibility.

diff --git a/Engine/Config/BaseEngine.ini b/Engine/Config/BaseEngine.ini
--- a/Engine/Config/BaseEngine.ini
+++ b/Engine/Config/BaseEngine.ini
@@ -1095,6 +1095,9 @@ DefaultWorldPartitionHLODLayer=/Engine/Maps/Templates/HLODs/HLODLayer_Instanced.
 +FunctionRedirects=(OldName="LegacyCameraShake.StartMatineeCameraShakeFromSource", NewName="LegacyCameraShake.StartLegacyCameraShakeFromSource")
 +FunctionRedirects=(OldName="LegacyCameraShakeFunctionLibrary.Conv_MatineeCameraShake", NewName="LegacyCameraShakeFunctionLibrary.Conv_LegacyCameraShake")
 
+; 5.2
++PropertyRedirects=(OldName="DecalComponent.SetFadeIn.Duaration", NewName="DecalComponent.SetFadeIn.Duration")
+
 [CoreUObject.Metadata]
 ; Note: UnrealHeaderTool should be rerun after making changes to MetadataRedirects to catch any keys specified in class headers
 +MetadataRedirects=(OldKey="K2Protected", NewKey="BlueprintProtected")
diff --git a/Engine/Source/Runtime/Engine/Classes/Components/DecalComponent.h b/Engine/Source/Runtime/Engine/Classes/Components/DecalComponent.h
--- a/Engine/Source/Runtime/Engine/Classes/Components/DecalComponent.h
+++ b/Engine/Source/Runtime/Engine/Classes/Components/DecalComponent.h
@@ -88,7 +88,7 @@ public:
 	void SetFadeOut(float StartDelay, float Duration, bool DestroyOwnerAfterFade = true);
 
 	UFUNCTION(BlueprintCallable, Category = "Rendering|Components|Decal")
-	void SetFadeIn(float StartDelay, float Duaration);
+	void SetFadeIn(float StartDelay, float Duration);
 
 	/** Set the FadeScreenSize for this decal component */
 	UFUNCTION(BlueprintCallable, Category = "Rendering|Components|Decal")
