Adding entry for cvar '' which will soon be replace the existing cvar '', by making the changes now we can avoid any problem with the change over.

#rnx

diff --git a/Engine/Config/BaseDeviceProfiles.ini b/Engine/Config/BaseDeviceProfiles.ini
--- a/Engine/Config/BaseDeviceProfiles.ini
+++ b/Engine/Config/BaseDeviceProfiles.ini
@@ -165,7 +165,9 @@ DeviceType=Windows
 BaseProfileName=Windows
 +CVars=s.AllowUnversionedContentInEditor=1
 +CVars=cook.AllowCookedDataInEditorBuilds=1
+; TODO - Remove VA.DisableVirtualization once it has been replaced with VA.AllowPkgVirtualization in code
 +CVars=VA.DisableVirtualization=1
++CVars=VA.AllowPkgVirtualization=0
 
 ;----------------
 ; iOS Devices
@@ -955,7 +957,9 @@ DeviceType=Linux
 BaseProfileName=Linux
 +CVars=s.AllowUnversionedContentInEditor=1
 +CVars=cook.AllowCookedDataInEditorBuilds=1
+; TODO - Remove VA.DisableVirtualization once it has been replaced with VA.AllowPkgVirtualization in code
 +CVars=VA.DisableVirtualization=1
++CVars=VA.AllowPkgVirtualization=0
 
 [HoloLens DeviceProfile]
 DeviceType=HoloLens
