Add missing property redirects related to CL 21790683

diff --git a/Engine/Config/BaseEngine.ini b/Engine/Config/BaseEngine.ini
--- a/Engine/Config/BaseEngine.ini
+++ b/Engine/Config/BaseEngine.ini
@@ -1095,6 +1095,10 @@ DefaultWorldPartitionHLODLayer=/Engine/Maps/Templates/HLODs/HLODLayer_Instanced.
 +FunctionRedirects=(OldName="LegacyCameraShake.StartMatineeCameraShakeFromSource", NewName="LegacyCameraShake.StartLegacyCameraShakeFromSource")
 +FunctionRedirects=(OldName="LegacyCameraShakeFunctionLibrary.Conv_MatineeCameraShake", NewName="LegacyCameraShakeFunctionLibrary.Conv_LegacyCameraShake")
 
++PropertyRedirects=(OldName="AnimationSettings.BoneTimecodeCustomAttributeNameSettings",NewName="BoneTimecodeAnimationAttributeNameSettings")
++PropertyRedirects=(OldName="AnimationSettings.BoneCustomAttributesNames",NewName="BoneAnimationAttributesNames")
++PropertyRedirects=(OldName="AnimationSettings.BoneNamesWithCustomAttributes",NewName="BoneNamesWithAnimationAttributes")
+
 ; 5.2
 +PropertyRedirects=(OldName="DecalComponent.SetFadeIn.Duaration", NewName="DecalComponent.SetFadeIn.Duration")
 
