Remove entries to the now removed cvar 'VA.DisableVirtualization'

#rnx

diff --git a/Engine/Config/BaseDeviceProfiles.ini b/Engine/Config/BaseDeviceProfiles.ini
--- a/Engine/Config/BaseDeviceProfiles.ini
+++ b/Engine/Config/BaseDeviceProfiles.ini
@@ -165,8 +165,6 @@ DeviceType=Windows
 BaseProfileName=Windows
 +CVars=s.AllowUnversionedContentInEditor=1
 +CVars=cook.AllowCookedDataInEditorBuilds=1
-; TODO - Remove VA.DisableVirtualization once it has been replaced with VA.AllowPkgVirtualization in code
-+CVars=VA.DisableVirtualization=1
 +CVars=VA.AllowPkgVirtualization=0
 
 ;----------------
@@ -957,8 +955,6 @@ DeviceType=Linux
 BaseProfileName=Linux
 +CVars=s.AllowUnversionedContentInEditor=1
 +CVars=cook.AllowCookedDataInEditorBuilds=1
-; TODO - Remove VA.DisableVirtualization once it has been replaced with VA.AllowPkgVirtualization in code
-+CVars=VA.DisableVirtualization=1
 +CVars=VA.AllowPkgVirtualization=0
 
 [HoloLens DeviceProfile]
