Enable Lumen Radiosity on High GI settings, lower quality

diff --git a/Engine/Config/BaseScalability.ini b/Engine/Config/BaseScalability.ini
--- a/Engine/Config/BaseScalability.ini
+++ b/Engine/Config/BaseScalability.ini
@@ -259,8 +259,8 @@ r.Lumen.DiffuseIndirect.Allow=0
 
 [GlobalIlluminationQuality@2]
 r.Lumen.DiffuseIndirect.Allow=1
-r.Lumen.TranslucencyVolume.GridPixelSize=64
-r.LumenScene.Radiosity=0
+r.LumenScene.Radiosity.ProbeSpacing=8
+r.LumenScene.Radiosity.HemisphereProbeResolution=3
 r.Lumen.TraceMeshSDFs.Allow=0
 r.Lumen.ScreenProbeGather.RadianceCache.ProbeResolution=16
 r.Lumen.ScreenProbeGather.RadianceCache.NumProbesToTraceBudget=300
@@ -271,6 +271,7 @@ r.Lumen.ScreenProbeGather.IrradianceFormat=1
 r.Lumen.ScreenProbeGather.StochasticInterpolation=1
 r.Lumen.ScreenProbeGather.FullResolutionJitterWidth=1
 r.Lumen.ScreenProbeGather.TwoSidedFoliageBackfaceDiffuse=0
+r.Lumen.TranslucencyVolume.GridPixelSize=64
 r.Lumen.TranslucencyVolume.TraceFromVolume=0
 r.Lumen.TranslucencyVolume.RadianceCache.ProbeResolution=8
 r.Lumen.TranslucencyVolume.RadianceCache.NumProbesToTraceBudget=200
@@ -278,8 +279,8 @@ r.LumenScene.SurfaceCache.CardCaptureRefreshFraction=0
 
 [GlobalIlluminationQuality@3]
 r.Lumen.DiffuseIndirect.Allow=1
-r.Lumen.TranslucencyVolume.GridPixelSize=32
-r.LumenScene.Radiosity=1
+r.LumenScene.Radiosity.ProbeSpacing=4
+r.LumenScene.Radiosity.HemisphereProbeResolution=4
 r.Lumen.TraceMeshSDFs.Allow=1
 r.Lumen.ScreenProbeGather.RadianceCache.ProbeResolution=32
 r.Lumen.ScreenProbeGather.RadianceCache.NumProbesToTraceBudget=300
@@ -290,6 +291,7 @@ r.Lumen.ScreenProbeGather.IrradianceFormat=0
 r.Lumen.ScreenProbeGather.StochasticInterpolation=0
 r.Lumen.ScreenProbeGather.FullResolutionJitterWidth=1
 r.Lumen.ScreenProbeGather.TwoSidedFoliageBackfaceDiffuse=1
+r.Lumen.TranslucencyVolume.GridPixelSize=32
 r.Lumen.TranslucencyVolume.TraceFromVolume=1
 r.Lumen.TranslucencyVolume.TracingOctahedronResolution=3
 r.Lumen.TranslucencyVolume.RadianceCache.ProbeResolution=8
@@ -298,8 +300,8 @@ r.LumenScene.SurfaceCache.CardCaptureRefreshFraction=0.125
 
 [GlobalIlluminationQuality@Cine]
 r.Lumen.DiffuseIndirect.Allow=1
-r.Lumen.TranslucencyVolume.GridPixelSize=32
-r.LumenScene.Radiosity=1
+r.LumenScene.Radiosity.ProbeSpacing=4
+r.LumenScene.Radiosity.HemisphereProbeResolution=4
 r.Lumen.TraceMeshSDFs.Allow=1
 r.Lumen.ScreenProbeGather.RadianceCache.ProbeResolution=32
 r.Lumen.ScreenProbeGather.RadianceCache.NumProbesToTraceBudget=100000
@@ -310,6 +312,7 @@ r.Lumen.ScreenProbeGather.IrradianceFormat=0
 r.Lumen.ScreenProbeGather.StochasticInterpolation=0
 r.Lumen.ScreenProbeGather.FullResolutionJitterWidth=.5
 r.Lumen.ScreenProbeGather.TwoSidedFoliageBackfaceDiffuse=1
+r.Lumen.TranslucencyVolume.GridPixelSize=32
 r.Lumen.TranslucencyVolume.TraceFromVolume=1
 r.Lumen.TranslucencyVolume.TracingOctahedronResolution=4
 r.Lumen.TranslucencyVolume.RadianceCache.ProbeResolution=16
