Workaround Lumen Radiance Cache corruption on Cinematic scalability level

diff --git a/Engine/Config/BaseScalability.ini b/Engine/Config/BaseScalability.ini
--- a/Engine/Config/BaseScalability.ini
+++ b/Engine/Config/BaseScalability.ini
@@ -304,7 +304,7 @@ r.LumenScene.Radiosity.ProbeSpacing=4
 r.LumenScene.Radiosity.HemisphereProbeResolution=4
 r.Lumen.TraceMeshSDFs.Allow=1
 r.Lumen.ScreenProbeGather.RadianceCache.ProbeResolution=32
-r.Lumen.ScreenProbeGather.RadianceCache.NumProbesToTraceBudget=100000
+r.Lumen.ScreenProbeGather.RadianceCache.NumProbesToTraceBudget=1000
 r.Lumen.ScreenProbeGather.ScreenSpaceBentNormal=1
 r.Lumen.ScreenProbeGather.DownsampleFactor=8
 r.Lumen.ScreenProbeGather.TracingOctahedronResolution=16
@@ -316,7 +316,7 @@ r.Lumen.TranslucencyVolume.GridPixelSize=32
 r.Lumen.TranslucencyVolume.TraceFromVolume=1
 r.Lumen.TranslucencyVolume.TracingOctahedronResolution=4
 r.Lumen.TranslucencyVolume.RadianceCache.ProbeResolution=16
-r.Lumen.TranslucencyVolume.RadianceCache.NumProbesToTraceBudget=100000
+r.Lumen.TranslucencyVolume.RadianceCache.NumProbesToTraceBudget=1000
 r.LumenScene.SurfaceCache.CardCaptureRefreshFraction=0.125
 
 
