Disabled cloud beer shadow map spatial filtering in game, except for cinematic.

diff --git a/Engine/Config/BaseScalability.ini b/Engine/Config/BaseScalability.ini
--- a/Engine/Config/BaseScalability.ini
+++ b/Engine/Config/BaseScalability.ini
@@ -603,6 +603,7 @@ r.SkyAtmosphere.TransmittanceLUT.UseSmallFormat=1
 r.SkyAtmosphere.TransmittanceLUT.SampleCount=10.0
 r.SkyAtmosphere.MultiScatteringLUT.SampleCount=15.0
 r.SkyLight.RealTimeReflectionCapture=0
+r.VolumetricCloud.ShadowMap.SpatialFiltering=0
 fx.Niagara.QualityLevel=0
 
 [EffectsQuality@1]
@@ -633,6 +634,7 @@ r.SkyAtmosphere.TransmittanceLUT.UseSmallFormat=0
 r.SkyAtmosphere.TransmittanceLUT.SampleCount=10.0
 r.SkyAtmosphere.MultiScatteringLUT.SampleCount=15.0
 r.SkyLight.RealTimeReflectionCapture=0
+r.VolumetricCloud.ShadowMap.SpatialFiltering=0
 fx.Niagara.QualityLevel=1
 
 [EffectsQuality@2]
@@ -663,6 +665,7 @@ r.SkyAtmosphere.TransmittanceLUT.UseSmallFormat=0
 r.SkyAtmosphere.TransmittanceLUT.SampleCount=10.0
 r.SkyAtmosphere.MultiScatteringLUT.SampleCount=15.0
 r.SkyLight.RealTimeReflectionCapture=1
+r.VolumetricCloud.ShadowMap.SpatialFiltering=0
 fx.Niagara.QualityLevel=2
 
 [EffectsQuality@3]
@@ -693,6 +696,7 @@ r.SkyAtmosphere.TransmittanceLUT.UseSmallFormat=0
 r.SkyAtmosphere.TransmittanceLUT.SampleCount=10.0
 r.SkyAtmosphere.MultiScatteringLUT.SampleCount=15.0
 r.SkyLight.RealTimeReflectionCapture=1
+r.VolumetricCloud.ShadowMap.SpatialFiltering=0
 fx.Niagara.QualityLevel=3
 
 [EffectsQuality@Cine]
@@ -723,6 +727,7 @@ r.SkyAtmosphere.TransmittanceLUT.UseSmallFormat=0
 r.SkyAtmosphere.TransmittanceLUT.SampleCount=30.0
 r.SkyAtmosphere.MultiScatteringLUT.SampleCount=20.0
 r.SkyLight.RealTimeReflectionCapture=1
+r.VolumetricCloud.ShadowMap.SpatialFiltering=1
 fx.Niagara.QualityLevel=4
 
 ;-----------------------------------------------------------------------------------------------------------------
