IncrementalCooking: Add DeclareConstructClasses for classes that need it. Fix some transient UObjects that assert they are never persistently serialized to support Serialize being called on their ClassDefaultObject to support DeclareCustomVersions.
#rnx

diff --git a/Engine/Config/BaseEditor.ini b/Engine/Config/BaseEditor.ini
--- a/Engine/Config/BaseEditor.ini
+++ b/Engine/Config/BaseEditor.ini
@@ -398,24 +398,7 @@ SaveUnversioned=true
 +GlobalCanConstructClasses=/Script/CoreUObject.MetaData
 +GlobalCanConstructClasses=/Script/Engine.AssetImportData
 ; +PostLoadCanConstructClasses=<LongPackageNameOfExistingClass>,<LongPackageNameOfClassItCanCreateInPostLoad>
-+PostLoadCanConstructClasses=/Script/ClothingSystemRuntimeCommon.ClothingAssetCommon,/Script/ChaosCloth.ChaosClothSharedSimConfig
-+PostLoadCanConstructClasses=/Script/ClothingSystemRuntimeCommon.ClothingAssetCommon,/Script/ChaosCloth.ChaosClothConfig
-+PostLoadCanConstructClasses=/Script/Engine.AnimSequenceBase,/Script/Engine.AnimDataModel
-+PostLoadCanConstructClasses=/Script/Engine.MapBuildDataRegistry,/Script/Engine.TextureCube
-+PostLoadCanConstructClasses=/Script/Engine.PhysicsAsset,/Script/Engine.SkeletalBodySetup
-+PostLoadCanConstructClasses=/Script/Engine.SkeletalMesh,/Script/Engine.SkeletalMeshEditorData
-+PostLoadCanConstructClasses=/Script/Engine.StaticMesh,/Script/Engine.StaticMeshDescriptionBulkData
-+PostLoadCanConstructClasses=/Script/LevelSequence.LevelSequenceActor,/Script/LevelSequence.DefaultLevelSequenceInstanceData
-+PostLoadCanConstructClasses=/Script/MovieScene.MovieScene,/Script/MovieScene.MovieSceneNodeGroupCollection
-+PostLoadCanConstructClasses=/Script/MovieScene.MovieSceneSection,/Script/MovieScene.MovieSceneBuiltInEasingFunction
-+PostLoadCanConstructClasses=/Script/Niagara.NiagaraEffectType,/Script/Niagara.NiagaraSignificanceHandlerDistance
-+PostLoadCanConstructClasses=/Script/Niagara.NiagaraEmitter,/Script/NiagaraEditor.NiagaraEditorParametersAdapter
-+PostLoadCanConstructClasses=/Script/Niagara.NiagaraSystem,/Script/NiagaraEditor.NiagaraEditorParametersAdapter
-+PostLoadCanConstructClasses=/Script/NiagaraEditor.NiagaraEmitterEditorData,/Script/Engine.EdGraph
-+PostLoadCanConstructClasses=/Script/NiagaraEditor.NiagaraGraph,/Script/NiagaraEditor.NiagaraScriptVariable
-+PostLoadCanConstructClasses=/Script/NiagaraEditor.NiagaraSystemEditorData,/Script/NiagaraEditor.NiagaraOverviewNode
-
-
+; PostLoadCanConstructClasses is usually not needed because it can be declared in UCPPClass::DeclareConstructClasses
 
 
 [TargetDomain]
diff --git a/Engine/Plugins/Animation/AnimationData/Source/AnimationData/Private/AnimSequencerController.cpp b/Engine/Plugins/Animation/AnimationData/Source/AnimationData/Private/AnimSequencerController.cpp
--- a/Engine/Plugins/Animation/AnimationData/Source/AnimationData/Private/AnimSequencerController.cpp
+++ b/Engine/Plugins/Animation/AnimationData/Source/AnimationData/Private/AnimSequencerController.cpp
@@ -2440,6 +2440,16 @@ void UAnimSequencerController::PopulateWithExistingModel(TScriptInterface<IAnima
 	CloseBracket();
 }
 
+#if WITH_EDITORONLY_DATA
+void UAnimSequencerController::DeclareConstructClasses(TArray<FTopLevelAssetPath>& OutConstructClasses, const UClass* SpecificSubclass)
+{
+	Super::DeclareConstructClasses(OutConstructClasses, SpecificSubclass);
+	OutConstructClasses.Add(FTopLevelAssetPath(UMovieScene::StaticClass()));
+	OutConstructClasses.Add(FTopLevelAssetPath(UMovieSceneControlRigParameterTrack::StaticClass()));
+	OutConstructClasses.Add(FTopLevelAssetPath(UFKControlRig::StaticClass()));
+}
+#endif
+
 void UAnimSequencerController::InitializeModel()
 {
 	if (Model->GetAnimationSequence())
diff --git a/Engine/Plugins/Animation/AnimationData/Source/AnimationData/Private/AnimSequencerController.h b/Engine/Plugins/Animation/AnimationData/Source/AnimationData/Private/AnimSequencerController.h
--- a/Engine/Plugins/Animation/AnimationData/Source/AnimationData/Private/AnimSequencerController.h
+++ b/Engine/Plugins/Animation/AnimationData/Source/AnimationData/Private/AnimSequencerController.h
@@ -18,6 +18,12 @@ public:
 	: BracketDepth(0) 
 	{}
 
+	/** Begin UObject Interface */
+#if WITH_EDITORONLY_DATA
+	static void DeclareConstructClasses(TArray<FTopLevelAssetPath>& OutConstructClasses, const UClass* SpecificSubclass);
+#endif
+	/** End UObject overrides */
+
 	/** Begin IAnimationDataController overrides */
 	virtual void SetModel(TScriptInterface<IAnimationDataModel> InModel) override;
     virtual TScriptInterface<IAnimationDataModel> GetModelInterface() const override { return ModelInterface; }
diff --git a/Engine/Plugins/Animation/AnimationData/Source/AnimationData/Private/AnimSequencerDataModel.cpp b/Engine/Plugins/Animation/AnimationData/Source/AnimationData/Private/AnimSequencerDataModel.cpp
--- a/Engine/Plugins/Animation/AnimationData/Source/AnimationData/Private/AnimSequencerDataModel.cpp
+++ b/Engine/Plugins/Animation/AnimationData/Source/AnimationData/Private/AnimSequencerDataModel.cpp
@@ -180,6 +180,14 @@ void UAnimationSequencerDataModel::PostLoad()
 	}
 }
 
+#if WITH_EDITORONLY_DATA
+void UAnimationSequencerDataModel::DeclareConstructClasses(TArray<FTopLevelAssetPath>& OutConstructClasses, const UClass* SpecificSubclass)
+{
+	Super::DeclareConstructClasses(OutConstructClasses, SpecificSubclass);
+	OutConstructClasses.Add(FTopLevelAssetPath(UControlRig::StaticClass()));
+}
+#endif
+
 void UAnimationSequencerDataModel::GetPreloadDependencies(TArray<UObject*>& OutDeps)
 {
 	Super::GetPreloadDependencies(OutDeps);
diff --git a/Engine/Plugins/Animation/AnimationData/Source/AnimationData/Private/AnimSequencerDataModel.h b/Engine/Plugins/Animation/AnimationData/Source/AnimationData/Private/AnimSequencerDataModel.h
--- a/Engine/Plugins/Animation/AnimationData/Source/AnimationData/Private/AnimSequencerDataModel.h
+++ b/Engine/Plugins/Animation/AnimationData/Source/AnimationData/Private/AnimSequencerDataModel.h
@@ -56,6 +56,9 @@ public:
 	virtual void UnbindObjects(const FGuid& ObjectId, const TArray<UObject*>& InObjects, UObject* Context) override {}
 	virtual void UnbindInvalidObjects(const FGuid& ObjectId, UObject* Context) override {}
 	virtual void PostLoad() override;
+#if WITH_EDITORONLY_DATA
+	static void DeclareConstructClasses(TArray<FTopLevelAssetPath>& OutConstructClasses, const UClass* SpecificSubclass);
+#endif
 	/** End UMovieSceneSequence overrides */
 	
 	/** Begin IAnimationDataModel overrides*/
diff --git a/Engine/Plugins/Animation/ControlRig/Source/ControlRig/Private/ControlRig.cpp b/Engine/Plugins/Animation/ControlRig/Source/ControlRig/Private/ControlRig.cpp
--- a/Engine/Plugins/Animation/ControlRig/Source/ControlRig/Private/ControlRig.cpp
+++ b/Engine/Plugins/Animation/ControlRig/Source/ControlRig/Private/ControlRig.cpp
@@ -3370,6 +3370,16 @@ bool UControlRig::AreNativizedVMsDisabled()
 	return (CVarControlRigDisableNativizedVMs->GetInt() != 0);
 }
 
+#if WITH_EDITORONLY_DATA
+void UControlRig::DeclareConstructClasses(TArray<FTopLevelAssetPath>& OutConstructClasses, const UClass* SpecificSubclass)
+{
+	Super::DeclareConstructClasses(OutConstructClasses, SpecificSubclass);
+	OutConstructClasses.Add(FTopLevelAssetPath(URigVM::StaticClass()));
+	OutConstructClasses.Add(FTopLevelAssetPath(URigHierarchy::StaticClass()));
+	OutConstructClasses.Add(FTopLevelAssetPath(UAnimationDataSourceRegistry::StaticClass()));
+}
+#endif
+
 void UControlRig::PostInitInstance(UControlRig* InCDO)
 {
 	const EObjectFlags SubObjectFlags =
diff --git a/Engine/Plugins/Animation/ControlRig/Source/ControlRig/Private/Rigs/RigHierarchy.cpp b/Engine/Plugins/Animation/ControlRig/Source/ControlRig/Private/Rigs/RigHierarchy.cpp
--- a/Engine/Plugins/Animation/ControlRig/Source/ControlRig/Private/Rigs/RigHierarchy.cpp
+++ b/Engine/Plugins/Animation/ControlRig/Source/ControlRig/Private/Rigs/RigHierarchy.cpp
@@ -290,6 +290,14 @@ void URigHierarchy::PostLoad()
 	CleanupInvalidCaches();
 }
 
+#if WITH_EDITORONLY_DATA
+void URigHierarchy::DeclareConstructClasses(TArray<FTopLevelAssetPath>& OutConstructClasses, const UClass* SpecificSubclass)
+{
+	Super::DeclareConstructClasses(OutConstructClasses, SpecificSubclass);
+	OutConstructClasses.Add(FTopLevelAssetPath(URigHierarchyController::StaticClass()));
+}
+#endif
+
 void URigHierarchy::Reset()
 {
 	LLM_SCOPE_BYNAME(TEXT("Animation/ControlRig"));
diff --git a/Engine/Plugins/Animation/ControlRig/Source/ControlRig/Private/Sequencer/MovieSceneControlRigParameterTrack.cpp b/Engine/Plugins/Animation/ControlRig/Source/ControlRig/Private/Sequencer/MovieSceneControlRigParameterTrack.cpp
--- a/Engine/Plugins/Animation/ControlRig/Source/ControlRig/Private/Sequencer/MovieSceneControlRigParameterTrack.cpp
+++ b/Engine/Plugins/Animation/ControlRig/Source/ControlRig/Private/Sequencer/MovieSceneControlRigParameterTrack.cpp
@@ -422,6 +422,15 @@ void UMovieSceneControlRigParameterTrack::PostLoad()
 #endif
 }
 
+#if WITH_EDITORONLY_DATA
+void UMovieSceneControlRigParameterTrack::DeclareConstructClasses(TArray<FTopLevelAssetPath>& OutConstructClasses, const UClass* SpecificSubclass)
+{
+	Super::DeclareConstructClasses(OutConstructClasses, SpecificSubclass);
+	OutConstructClasses.Add(FTopLevelAssetPath(UMovieSceneSection::StaticClass()));
+	OutConstructClasses.Add(FTopLevelAssetPath(UMovieSceneControlRigParameterSection::StaticClass()));
+}
+#endif
+
 #if WITH_EDITOR
 void UMovieSceneControlRigParameterTrack::HandlePackageDone(const FEndLoadPackageContext& Context)
 {
diff --git a/Engine/Plugins/Animation/ControlRig/Source/ControlRig/Public/ControlRig.h b/Engine/Plugins/Animation/ControlRig/Source/ControlRig/Public/ControlRig.h
--- a/Engine/Plugins/Animation/ControlRig/Source/ControlRig/Public/ControlRig.h
+++ b/Engine/Plugins/Animation/ControlRig/Source/ControlRig/Public/ControlRig.h
@@ -494,6 +494,9 @@ public:
 	void PostInitInstanceIfRequired();
 	void SwapVMToNativizedIfRequired(UClass* InNativizedClass = nullptr);
 	static bool AreNativizedVMsDisabled() ;
+#if WITH_EDITORONLY_DATA
+	static void DeclareConstructClasses(TArray<FTopLevelAssetPath>& OutConstructClasses, const UClass* SpecificSubclass);
+#endif
 
 protected:
 
diff --git a/Engine/Plugins/Animation/ControlRig/Source/ControlRig/Public/Rigs/RigHierarchy.h b/Engine/Plugins/Animation/ControlRig/Source/ControlRig/Public/Rigs/RigHierarchy.h
--- a/Engine/Plugins/Animation/ControlRig/Source/ControlRig/Public/Rigs/RigHierarchy.h
+++ b/Engine/Plugins/Animation/ControlRig/Source/ControlRig/Public/Rigs/RigHierarchy.h
@@ -107,6 +107,9 @@ public:
 	void Save(FArchive& Ar);
 	void Load(FArchive& Ar);
 	virtual void PostLoad() override;
+#if WITH_EDITORONLY_DATA
+	static void DeclareConstructClasses(TArray<FTopLevelAssetPath>& OutConstructClasses, const UClass* SpecificSubclass);
+#endif
 #if WITH_EDITOR
 	virtual void PostEditUndo() override;
 #endif
diff --git a/Engine/Plugins/Animation/ControlRig/Source/ControlRig/Public/Sequencer/MovieSceneControlRigParameterTrack.h b/Engine/Plugins/Animation/ControlRig/Source/ControlRig/Public/Sequencer/MovieSceneControlRigParameterTrack.h
--- a/Engine/Plugins/Animation/ControlRig/Source/ControlRig/Public/Sequencer/MovieSceneControlRigParameterTrack.h
+++ b/Engine/Plugins/Animation/ControlRig/Source/ControlRig/Public/Sequencer/MovieSceneControlRigParameterTrack.h
@@ -42,6 +42,9 @@ public:
 	virtual FName GetTrackName() const override { return TrackName; }
 	// UObject
 	virtual void PostLoad() override;
+#if WITH_EDITORONLY_DATA
+	static void DeclareConstructClasses(TArray<FTopLevelAssetPath>& OutConstructClasses, const UClass* SpecificSubclass);
+#endif
 	virtual void PostEditImport() override;
 #if WITH_EDITORONLY_DATA
 	virtual FText GetDefaultDisplayName() const override;
diff --git a/Engine/Plugins/Animation/ControlRig/Source/ControlRigDeveloper/Private/ControlRigBlueprint.cpp b/Engine/Plugins/Animation/ControlRig/Source/ControlRigDeveloper/Private/ControlRigBlueprint.cpp
--- a/Engine/Plugins/Animation/ControlRig/Source/ControlRigDeveloper/Private/ControlRigBlueprint.cpp
+++ b/Engine/Plugins/Animation/ControlRig/Source/ControlRigDeveloper/Private/ControlRigBlueprint.cpp
@@ -1047,6 +1047,14 @@ void UControlRigBlueprint::PostLoad()
 	FRigVMRegistry::Get().OnRigVMRegistryChanged().AddUObject(this, &UControlRigBlueprint::OnRigVMRegistryChanged);
 }
 
+#if WITH_EDITORONLY_DATA
+void UControlRigBlueprint::DeclareConstructClasses(TArray<FTopLevelAssetPath>& OutConstructClasses, const UClass* SpecificSubclass)
+{
+	Super::DeclareConstructClasses(OutConstructClasses, SpecificSubclass);
+	OutConstructClasses.Add(FTopLevelAssetPath(URigVMController::StaticClass()));
+}
+#endif
+
 #if WITH_EDITOR
 void UControlRigBlueprint::HandlePackageDone(const FEndLoadPackageContext& Context)
 {
diff --git a/Engine/Plugins/Animation/ControlRig/Source/ControlRigDeveloper/Public/ControlRigBlueprint.h b/Engine/Plugins/Animation/ControlRig/Source/ControlRigDeveloper/Public/ControlRigBlueprint.h
--- a/Engine/Plugins/Animation/ControlRig/Source/ControlRigDeveloper/Public/ControlRigBlueprint.h
+++ b/Engine/Plugins/Animation/ControlRig/Source/ControlRigDeveloper/Public/ControlRigBlueprint.h
@@ -236,6 +236,9 @@ public:
 	PRAGMA_ENABLE_DEPRECATION_WARNINGS
 	virtual void PreSave(FObjectPreSaveContext ObjectSaveContext) override;
 	virtual void PostLoad() override;
+#if WITH_EDITORONLY_DATA
+	static void DeclareConstructClasses(TArray<FTopLevelAssetPath>& OutConstructClasses, const UClass* SpecificSubclass);
+#endif
 	virtual bool IsPostLoadThreadSafe() const override { return false; }
 	virtual void PostTransacted(const FTransactionObjectEvent& TransactionEvent) override;
 	virtual void ReplaceDeprecatedNodes() override;
diff --git a/Engine/Plugins/Compositing/OpenColorIO/Source/OpenColorIO/Private/OpenColorIOColorTransform.cpp b/Engine/Plugins/Compositing/OpenColorIO/Source/OpenColorIO/Private/OpenColorIOColorTransform.cpp
--- a/Engine/Plugins/Compositing/OpenColorIO/Source/OpenColorIO/Private/OpenColorIOColorTransform.cpp
+++ b/Engine/Plugins/Compositing/OpenColorIO/Source/OpenColorIO/Private/OpenColorIOColorTransform.cpp
@@ -832,6 +832,15 @@ void UOpenColorIOColorTransform::PostLoad()
 
 }
 
+#if WITH_EDITORONLY_DATA
+void UOpenColorIOColorTransform::DeclareConstructClasses(TArray<FTopLevelAssetPath>& OutConstructClasses, const UClass* SpecificSubclass)
+{
+	Super::DeclareConstructClasses(OutConstructClasses, SpecificSubclass);
+	OutConstructClasses.Add(FTopLevelAssetPath(UTexture2D::StaticClass()));
+	OutConstructClasses.Add(FTopLevelAssetPath(UVolumeTexture::StaticClass()));
+}
+#endif
+
 void UOpenColorIOColorTransform::BeginDestroy()
 {
 	Super::BeginDestroy();
diff --git a/Engine/Plugins/Compositing/OpenColorIO/Source/OpenColorIO/Private/OpenColorIOConfiguration.cpp b/Engine/Plugins/Compositing/OpenColorIO/Source/OpenColorIO/Private/OpenColorIOConfiguration.cpp
--- a/Engine/Plugins/Compositing/OpenColorIO/Source/OpenColorIO/Private/OpenColorIOConfiguration.cpp
+++ b/Engine/Plugins/Compositing/OpenColorIO/Source/OpenColorIO/Private/OpenColorIOConfiguration.cpp
@@ -425,6 +425,14 @@ void UOpenColorIOConfiguration::PostLoad()
 	}
 }
 
+#if WITH_EDITORONLY_DATA
+void UOpenColorIOConfiguration::DeclareConstructClasses(TArray<FTopLevelAssetPath>& OutConstructClasses, const UClass* SpecificSubclass)
+{
+	Super::DeclareConstructClasses(OutConstructClasses, SpecificSubclass);
+	OutConstructClasses.Add(FTopLevelAssetPath(UOpenColorIOColorTransform::StaticClass()));
+}
+#endif
+
 namespace OpenColorIOConfiguration
 {
 	static void SendAnalytics(const FString& EventName, const TArray<FOpenColorIOColorSpace>& DesiredColorSpaces)
diff --git a/Engine/Plugins/Compositing/OpenColorIO/Source/OpenColorIO/Public/OpenColorIOColorTransform.h b/Engine/Plugins/Compositing/OpenColorIO/Source/OpenColorIO/Public/OpenColorIOColorTransform.h
--- a/Engine/Plugins/Compositing/OpenColorIO/Source/OpenColorIO/Public/OpenColorIOColorTransform.h
+++ b/Engine/Plugins/Compositing/OpenColorIO/Source/OpenColorIO/Public/OpenColorIOColorTransform.h
@@ -148,6 +148,10 @@ public:
 	//~ Begin UObject interface
 	void Serialize(FArchive& Ar) override;
 	virtual void PostLoad() override;
+#if WITH_EDITORONLY_DATA
+	static void DeclareConstructClasses(TArray<FTopLevelAssetPath>& OutConstructClasses, const UClass* SpecificSubclass);
+#endif
+
 	virtual void BeginDestroy() override;
 	virtual bool IsReadyForFinishDestroy() override;
 	virtual void FinishDestroy() override;
diff --git a/Engine/Plugins/Compositing/OpenColorIO/Source/OpenColorIO/Public/OpenColorIOConfiguration.h b/Engine/Plugins/Compositing/OpenColorIO/Source/OpenColorIO/Public/OpenColorIOConfiguration.h
--- a/Engine/Plugins/Compositing/OpenColorIO/Source/OpenColorIO/Public/OpenColorIOConfiguration.h
+++ b/Engine/Plugins/Compositing/OpenColorIO/Source/OpenColorIO/Public/OpenColorIOConfiguration.h
@@ -72,6 +72,10 @@ public:
 
 	//~ Begin UObject interface
 	virtual void PostLoad() override;
+#if WITH_EDITORONLY_DATA
+	static void DeclareConstructClasses(TArray<FTopLevelAssetPath>& OutConstructClasses, const UClass* SpecificSubclass);
+#endif
+
 	virtual void PreSave(FObjectPreSaveContext SaveContext) override;
 	virtual void BeginDestroy() override;
 
diff --git a/Engine/Plugins/FX/Niagara/Source/Niagara/Classes/NiagaraBakerSettings.h b/Engine/Plugins/FX/Niagara/Source/Niagara/Classes/NiagaraBakerSettings.h
--- a/Engine/Plugins/FX/Niagara/Source/Niagara/Classes/NiagaraBakerSettings.h
+++ b/Engine/Plugins/FX/Niagara/Source/Niagara/Classes/NiagaraBakerSettings.h
@@ -213,6 +213,10 @@ public:
 	FNiagaraBakerOutputFrameIndices GetOutputFrameIndices(int OutputIndex, float RelativeTime) const;
 
 	virtual void PostLoad() override;
+#if WITH_EDITORONLY_DATA
+	static void DeclareConstructClasses(TArray<FTopLevelAssetPath>& OutConstructClasses, const UClass* SpecificSubclass);
+#endif
+
 #if WITH_EDITORONLY_DATA
 	virtual void PostEditChangeProperty(struct FPropertyChangedEvent& PropertyChangedEvent) override;
 #endif
diff --git a/Engine/Plugins/FX/Niagara/Source/Niagara/Classes/NiagaraEffectType.h b/Engine/Plugins/FX/Niagara/Source/Niagara/Classes/NiagaraEffectType.h
--- a/Engine/Plugins/FX/Niagara/Source/Niagara/Classes/NiagaraEffectType.h
+++ b/Engine/Plugins/FX/Niagara/Source/Niagara/Classes/NiagaraEffectType.h
@@ -407,6 +407,9 @@ class NIAGARA_API UNiagaraEffectType : public UObject
 	virtual bool IsReadyForFinishDestroy()override;
 	virtual void Serialize(FArchive& Ar)override;
 	virtual void PostLoad()override;
+#if WITH_EDITORONLY_DATA
+	static void DeclareConstructClasses(TArray<FTopLevelAssetPath>& OutConstructClasses, const UClass* SpecificSubclass);
+#endif
 #if WITH_EDITOR
 	virtual void PostEditChangeProperty(struct FPropertyChangedEvent& PropertyChangedEvent);
 #endif
diff --git a/Engine/Plugins/FX/Niagara/Source/Niagara/Classes/NiagaraEmitter.h b/Engine/Plugins/FX/Niagara/Source/Niagara/Classes/NiagaraEmitter.h
--- a/Engine/Plugins/FX/Niagara/Source/Niagara/Classes/NiagaraEmitter.h
+++ b/Engine/Plugins/FX/Niagara/Source/Niagara/Classes/NiagaraEmitter.h
@@ -358,6 +358,7 @@ struct NIAGARA_API FVersionedNiagaraEmitterData
 
 	void CopyFrom(const FVersionedNiagaraEmitterData& Source);
 	void PostLoad(UNiagaraEmitter& Emitter, bool bIsCooked, int32 NiagaraVer);
+
 	void PostInitProperties(UNiagaraEmitter* Outer);
 	bool UsesCollection(const UNiagaraParameterCollection* Collection) const;
 	bool UsesScript(const UNiagaraScript* Script) const;
@@ -578,6 +579,9 @@ public:
 	virtual void Serialize(FArchive& Ar)override;
 	virtual void PostInitProperties() override;
 	virtual void PostLoad() override;
+#if WITH_EDITORONLY_DATA
+	static void DeclareConstructClasses(TArray<FTopLevelAssetPath>& OutConstructClasses, const UClass* SpecificSubclass);
+#endif
 	virtual bool IsEditorOnly() const override;
 	virtual void GetAssetRegistryTags(TArray<FAssetRegistryTag>& OutTags) const override;
 	//End UObject Interface
diff --git a/Engine/Plugins/FX/Niagara/Source/Niagara/Classes/NiagaraSystem.h b/Engine/Plugins/FX/Niagara/Source/Niagara/Classes/NiagaraSystem.h
--- a/Engine/Plugins/FX/Niagara/Source/Niagara/Classes/NiagaraSystem.h
+++ b/Engine/Plugins/FX/Niagara/Source/Niagara/Classes/NiagaraSystem.h
@@ -212,6 +212,9 @@ public:
 	void PostInitProperties();
 	virtual void Serialize(FArchive& Ar) override;
 	virtual void PostLoad() override; 
+#if WITH_EDITORONLY_DATA
+	static void DeclareConstructClasses(TArray<FTopLevelAssetPath>& OutConstructClasses, const UClass* SpecificSubclass);
+#endif
 	virtual void BeginDestroy() override;
 	PRAGMA_DISABLE_DEPRECATION_WARNINGS // Suppress compiler warning on override of deprecated function
 	UE_DEPRECATED(5.0, "Use version that takes FObjectPreSaveContext instead.")
diff --git a/Engine/Plugins/FX/Niagara/Source/Niagara/Private/NiagaraBakerSettings.cpp b/Engine/Plugins/FX/Niagara/Source/Niagara/Private/NiagaraBakerSettings.cpp
--- a/Engine/Plugins/FX/Niagara/Source/Niagara/Private/NiagaraBakerSettings.cpp
+++ b/Engine/Plugins/FX/Niagara/Source/Niagara/Private/NiagaraBakerSettings.cpp
@@ -217,6 +217,14 @@ void UNiagaraBakerSettings::PostLoad()
 	}
 }
 
+#if WITH_EDITORONLY_DATA
+void UNiagaraBakerSettings::DeclareConstructClasses(TArray<FTopLevelAssetPath>& OutConstructClasses, const UClass* SpecificSubclass)
+{
+	Super::DeclareConstructClasses(OutConstructClasses, SpecificSubclass);
+	OutConstructClasses.Add(FTopLevelAssetPath(UNiagaraBakerOutputTexture2D::StaticClass()));
+}
+#endif
+
 #if WITH_EDITORONLY_DATA
 void UNiagaraBakerSettings::PostEditChangeProperty(struct FPropertyChangedEvent& PropertyChangedEvent)
 {
diff --git a/Engine/Plugins/FX/Niagara/Source/Niagara/Private/NiagaraComponent.cpp b/Engine/Plugins/FX/Niagara/Source/Niagara/Private/NiagaraComponent.cpp
--- a/Engine/Plugins/FX/Niagara/Source/Niagara/Private/NiagaraComponent.cpp
+++ b/Engine/Plugins/FX/Niagara/Source/Niagara/Private/NiagaraComponent.cpp
@@ -2725,6 +2725,20 @@ void UNiagaraComponent::PostLoad()
 #endif
 }
 
+#if WITH_EDITORONLY_DATA
+void UNiagaraComponent::DeclareConstructClasses(TArray<FTopLevelAssetPath>& OutConstructClasses, const UClass* SpecificSubclass)
+{
+	Super::DeclareConstructClasses(OutConstructClasses, SpecificSubclass);
+	TArray<UClass*> DataInterfaceClasses;
+	GetDerivedClasses(UNiagaraDataInterface::StaticClass(), DataInterfaceClasses);
+	for (UClass* DataInterfaceClass : DataInterfaceClasses)
+	{
+		OutConstructClasses.Add(FTopLevelAssetPath(DataInterfaceClass));
+	}
+}
+#endif
+
+
 static FNiagaraVariant GetParameterValueFromStore(const FNiagaraVariableBase& Var, const FNiagaraParameterStore& Store)
 {
 	if (Var.IsDataInterface())
diff --git a/Engine/Plugins/FX/Niagara/Source/Niagara/Private/NiagaraEffectType.cpp b/Engine/Plugins/FX/Niagara/Source/Niagara/Private/NiagaraEffectType.cpp
--- a/Engine/Plugins/FX/Niagara/Source/Niagara/Private/NiagaraEffectType.cpp
+++ b/Engine/Plugins/FX/Niagara/Source/Niagara/Private/NiagaraEffectType.cpp
@@ -78,6 +78,14 @@ void UNiagaraEffectType::PostLoad()
 #endif
 }
 
+#if WITH_EDITORONLY_DATA
+void UNiagaraEffectType::DeclareConstructClasses(TArray<FTopLevelAssetPath>& OutConstructClasses, const UClass* SpecificSubclass)
+{
+	Super::DeclareConstructClasses(OutConstructClasses, SpecificSubclass);
+	OutConstructClasses.Add(FTopLevelAssetPath(UNiagaraSignificanceHandlerDistance::StaticClass()));
+}
+#endif
+
 const FNiagaraSystemScalabilitySettings& UNiagaraEffectType::GetActiveSystemScalabilitySettings()const
 {
 	for (const FNiagaraSystemScalabilitySettings& Settings : SystemScalabilitySettings.Settings)
diff --git a/Engine/Plugins/FX/Niagara/Source/Niagara/Private/NiagaraEmitter.cpp b/Engine/Plugins/FX/Niagara/Source/Niagara/Private/NiagaraEmitter.cpp
--- a/Engine/Plugins/FX/Niagara/Source/Niagara/Private/NiagaraEmitter.cpp
+++ b/Engine/Plugins/FX/Niagara/Source/Niagara/Private/NiagaraEmitter.cpp
@@ -651,6 +651,15 @@ void UNiagaraEmitter::PostLoad()
 #endif
 }
 
+#if WITH_EDITORONLY_DATA
+void UNiagaraEmitter::DeclareConstructClasses(TArray<FTopLevelAssetPath>& OutConstructClasses, const UClass* SpecificSubclass)
+{
+	Super::DeclareConstructClasses(OutConstructClasses, SpecificSubclass);
+	OutConstructClasses.Add(FTopLevelAssetPath(UNiagaraScratchPadContainer::StaticClass()));
+	OutConstructClasses.Add(FTopLevelAssetPath(TEXT("/Script/NiagaraEditor.NiagaraEditorParametersAdapter")));
+}
+#endif
+
 void FVersionedNiagaraEmitterData::PostLoad(UNiagaraEmitter& Emitter, bool bIsCooked, int32 NiagaraVer)
 {
 #if STATS
diff --git a/Engine/Plugins/FX/Niagara/Source/Niagara/Private/NiagaraSystem.cpp b/Engine/Plugins/FX/Niagara/Source/Niagara/Private/NiagaraSystem.cpp
--- a/Engine/Plugins/FX/Niagara/Source/Niagara/Private/NiagaraSystem.cpp
+++ b/Engine/Plugins/FX/Niagara/Source/Niagara/Private/NiagaraSystem.cpp
@@ -14,6 +14,7 @@
 #include "NiagaraModule.h"
 #include "NiagaraPrecompileContainer.h"
 #include "NiagaraRendererProperties.h"
+#include "NiagaraScratchPadContainer.h"
 #include "NiagaraScriptSourceBase.h"
 #include "NiagaraSettings.h"
 #include "NiagaraShared.h"
@@ -941,6 +942,16 @@ void UNiagaraSystem::PostLoad()
 	PrecachePSOs();
 }
 
+#if WITH_EDITORONLY_DATA
+void UNiagaraSystem::DeclareConstructClasses(TArray<FTopLevelAssetPath>& OutConstructClasses, const UClass* SpecificSubclass)
+{
+	Super::DeclareConstructClasses(OutConstructClasses, SpecificSubclass);
+	OutConstructClasses.Add(FTopLevelAssetPath(UNiagaraScratchPadContainer::StaticClass()));
+	OutConstructClasses.Add(FTopLevelAssetPath(TEXT("/Script/NiagaraEditor.NiagaraSystemEditorData")));
+	OutConstructClasses.Add(FTopLevelAssetPath(TEXT("/Script/NiagaraEditor.NiagaraEditorParametersAdapter")));
+}
+#endif
+
 void UNiagaraSystem::PrecachePSOs()
 {
 	if (!IsComponentPSOPrecachingEnabled() && !IsResourcePSOPrecachingEnabled())
diff --git a/Engine/Plugins/FX/Niagara/Source/Niagara/Public/NiagaraComponent.h b/Engine/Plugins/FX/Niagara/Source/Niagara/Public/NiagaraComponent.h
--- a/Engine/Plugins/FX/Niagara/Source/Niagara/Public/NiagaraComponent.h
+++ b/Engine/Plugins/FX/Niagara/Source/Niagara/Public/NiagaraComponent.h
@@ -564,6 +564,9 @@ public:
 	//~ Begin UObject Interface.
 	virtual void Serialize(FStructuredArchive::FRecord Record) override;
 	virtual void PostLoad() override;
+#if WITH_EDITORONLY_DATA
+	static void DeclareConstructClasses(TArray<FTopLevelAssetPath>& OutConstructClasses, const UClass* SpecificSubclass);
+#endif
 
 #if WITH_EDITOR
 	virtual void PreEditChange(FProperty* PropertyAboutToChange) override;
diff --git a/Engine/Plugins/FX/Niagara/Source/NiagaraEditor/Private/NiagaraEmitterEditorData.cpp b/Engine/Plugins/FX/Niagara/Source/NiagaraEditor/Private/NiagaraEmitterEditorData.cpp
--- a/Engine/Plugins/FX/Niagara/Source/NiagaraEditor/Private/NiagaraEmitterEditorData.cpp
+++ b/Engine/Plugins/FX/Niagara/Source/NiagaraEditor/Private/NiagaraEmitterEditorData.cpp
@@ -2,6 +2,7 @@
 
 #include "NiagaraEmitterEditorData.h"
 
+#include "EdGraph/EdGraph.h"
 #include "NiagaraStackEditorData.h"
 #include "ScopedTransaction.h"
 
@@ -33,6 +34,14 @@ void UNiagaraEmitterEditorData::PostLoad()
 	StackEditorData->ConditionalPostLoad();
 }
 
+#if WITH_EDITORONLY_DATA
+void UNiagaraEmitterEditorData::DeclareConstructClasses(TArray<FTopLevelAssetPath>& OutConstructClasses, const UClass* SpecificSubclass)
+{
+	Super::DeclareConstructClasses(OutConstructClasses, SpecificSubclass);
+	OutConstructClasses.Add(FTopLevelAssetPath(UEdGraph::StaticClass()));
+}
+#endif
+
 UNiagaraStackEditorData& UNiagaraEmitterEditorData::GetStackEditorData() const
 {
 	return *StackEditorData;
diff --git a/Engine/Plugins/FX/Niagara/Source/NiagaraEditor/Private/NiagaraGraph.cpp b/Engine/Plugins/FX/Niagara/Source/NiagaraEditor/Private/NiagaraGraph.cpp
--- a/Engine/Plugins/FX/Niagara/Source/NiagaraEditor/Private/NiagaraGraph.cpp
+++ b/Engine/Plugins/FX/Niagara/Source/NiagaraEditor/Private/NiagaraGraph.cpp
@@ -496,6 +496,14 @@ void UNiagaraGraph::PostLoad()
 	InvalidateCachedParameterData();
 }
 
+#if WITH_EDITORONLY_DATA
+void UNiagaraGraph::DeclareConstructClasses(TArray<FTopLevelAssetPath>& OutConstructClasses, const UClass* SpecificSubclass)
+{
+	Super::DeclareConstructClasses(OutConstructClasses, SpecificSubclass);
+	OutConstructClasses.Add(FTopLevelAssetPath(UNiagaraScriptVariable::StaticClass()));
+}
+#endif
+
 void UNiagaraGraph::ChangeParameterType(const TArray<FNiagaraVariable>& ParametersToChange, const FNiagaraTypeDefinition& NewType, bool bAllowOrphanedPins)
 {
 	check(!bIsForCompilationOnly);
diff --git a/Engine/Plugins/FX/Niagara/Source/NiagaraEditor/Private/NiagaraNodeAssignment.cpp b/Engine/Plugins/FX/Niagara/Source/NiagaraEditor/Private/NiagaraNodeAssignment.cpp
--- a/Engine/Plugins/FX/Niagara/Source/NiagaraEditor/Private/NiagaraNodeAssignment.cpp
+++ b/Engine/Plugins/FX/Niagara/Source/NiagaraEditor/Private/NiagaraNodeAssignment.cpp
@@ -247,6 +247,21 @@ void UNiagaraNodeAssignment::PostLoad()
 	RefreshTitle();
 }
 
+#if WITH_EDITORONLY_DATA
+void UNiagaraNodeAssignment::DeclareConstructClasses(TArray<FTopLevelAssetPath>& OutConstructClasses, const UClass* SpecificSubclass)
+{
+	Super::DeclareConstructClasses(OutConstructClasses, SpecificSubclass);
+	OutConstructClasses.Add(FTopLevelAssetPath(UNiagaraGraph::StaticClass()));
+	OutConstructClasses.Add(FTopLevelAssetPath(UNiagaraNodeInput::StaticClass()));	
+	OutConstructClasses.Add(FTopLevelAssetPath(UNiagaraNodeOutput::StaticClass()));
+	OutConstructClasses.Add(FTopLevelAssetPath(UNiagaraNodeParameterMapGet::StaticClass()));
+	OutConstructClasses.Add(FTopLevelAssetPath(UNiagaraNodeParameterMapSet::StaticClass()));
+	OutConstructClasses.Add(FTopLevelAssetPath(UNiagaraScript::StaticClass()));
+	OutConstructClasses.Add(FTopLevelAssetPath(UNiagaraScriptSource::StaticClass()));
+	OutConstructClasses.Add(FTopLevelAssetPath(UNiagaraScriptVariable::StaticClass()));	
+}
+#endif
+
 void UNiagaraNodeAssignment::BuildParameterMapHistory(FNiagaraParameterMapHistoryBuilder& OutHistory, bool bRecursive /*= true*/, bool bFilterForCompilation /*= true*/) const
 {
 	Super::BuildParameterMapHistory(OutHistory, bRecursive, bFilterForCompilation);
diff --git a/Engine/Plugins/FX/Niagara/Source/NiagaraEditor/Private/NiagaraSystemEditorData.cpp b/Engine/Plugins/FX/Niagara/Source/NiagaraEditor/Private/NiagaraSystemEditorData.cpp
--- a/Engine/Plugins/FX/Niagara/Source/NiagaraEditor/Private/NiagaraSystemEditorData.cpp
+++ b/Engine/Plugins/FX/Niagara/Source/NiagaraEditor/Private/NiagaraSystemEditorData.cpp
@@ -138,6 +138,17 @@ void UNiagaraSystemEditorData::PostLoadFromOwner(UObject* InOwner)
 	}
 }
 
+#if WITH_EDITORONLY_DATA
+void UNiagaraSystemEditorData::DeclareConstructClasses(TArray<FTopLevelAssetPath>& OutConstructClasses, const UClass* SpecificSubclass)
+{
+	Super::DeclareConstructClasses(OutConstructClasses, SpecificSubclass);
+	OutConstructClasses.Add(FTopLevelAssetPath(UEdGraph::StaticClass()));
+	OutConstructClasses.Add(FTopLevelAssetPath(UNiagaraOverviewNode::StaticClass()));
+	OutConstructClasses.Add(FTopLevelAssetPath(UNiagaraScriptVariable::StaticClass()));
+	OutConstructClasses.Add(FTopLevelAssetPath(UNiagaraHierarchyRoot::StaticClass()));
+}
+#endif
+
 UNiagaraSystemEditorFolder& UNiagaraSystemEditorData::GetRootFolder() const
 {
 	return *RootFolder;
diff --git a/Engine/Plugins/FX/Niagara/Source/NiagaraEditor/Public/NiagaraEmitterEditorData.h b/Engine/Plugins/FX/Niagara/Source/NiagaraEditor/Public/NiagaraEmitterEditorData.h
--- a/Engine/Plugins/FX/Niagara/Source/NiagaraEditor/Public/NiagaraEmitterEditorData.h
+++ b/Engine/Plugins/FX/Niagara/Source/NiagaraEditor/Public/NiagaraEmitterEditorData.h
@@ -26,6 +26,9 @@ public:
 	UNiagaraEmitterEditorData(const FObjectInitializer& ObjectInitializer);
 
 	virtual void PostLoad() override;
+#if WITH_EDITORONLY_DATA
+	NIAGARAEDITOR_API static void DeclareConstructClasses(TArray<FTopLevelAssetPath>& OutConstructClasses, const UClass* SpecificSubclass);
+#endif
 
 	NIAGARAEDITOR_API UNiagaraStackEditorData& GetStackEditorData() const;
 
diff --git a/Engine/Plugins/FX/Niagara/Source/NiagaraEditor/Public/NiagaraGraph.h b/Engine/Plugins/FX/Niagara/Source/NiagaraEditor/Public/NiagaraGraph.h
--- a/Engine/Plugins/FX/Niagara/Source/NiagaraEditor/Public/NiagaraGraph.h
+++ b/Engine/Plugins/FX/Niagara/Source/NiagaraEditor/Public/NiagaraGraph.h
@@ -148,6 +148,9 @@ class UNiagaraGraph : public UEdGraph
 
 	//~ Begin UObject Interface
 	virtual void PostLoad() override;
+#if WITH_EDITORONLY_DATA
+	static void DeclareConstructClasses(TArray<FTopLevelAssetPath>& OutConstructClasses, const UClass* SpecificSubclass);
+#endif
 	virtual void PostEditChangeProperty(struct FPropertyChangedEvent& PropertyChangedEvent) override;
 	virtual void BeginDestroy() override;
 	//~ End UObjet Interface
diff --git a/Engine/Plugins/FX/Niagara/Source/NiagaraEditor/Public/NiagaraNodeAssignment.h b/Engine/Plugins/FX/Niagara/Source/NiagaraEditor/Public/NiagaraNodeAssignment.h
--- a/Engine/Plugins/FX/Niagara/Source/NiagaraEditor/Public/NiagaraNodeAssignment.h
+++ b/Engine/Plugins/FX/Niagara/Source/NiagaraEditor/Public/NiagaraNodeAssignment.h
@@ -36,6 +36,9 @@ public:
 
 	//~ Begin EdGraphNode Interface
 	virtual void PostLoad() override;
+#if WITH_EDITORONLY_DATA
+	static void DeclareConstructClasses(TArray<FTopLevelAssetPath>& OutConstructClasses, const UClass* SpecificSubclass);
+#endif
 	virtual FText GetNodeTitle(ENodeTitleType::Type TitleType) const override;
 	virtual void AllocateDefaultPins() override;
 
diff --git a/Engine/Plugins/FX/Niagara/Source/NiagaraEditor/Public/NiagaraSystemEditorData.h b/Engine/Plugins/FX/Niagara/Source/NiagaraEditor/Public/NiagaraSystemEditorData.h
--- a/Engine/Plugins/FX/Niagara/Source/NiagaraEditor/Public/NiagaraSystemEditorData.h
+++ b/Engine/Plugins/FX/Niagara/Source/NiagaraEditor/Public/NiagaraSystemEditorData.h
@@ -58,6 +58,9 @@ public:
 	void PostInitProperties();
 
 	virtual void PostLoadFromOwner(UObject* InOwner) override;
+#if WITH_EDITORONLY_DATA
+	static void DeclareConstructClasses(TArray<FTopLevelAssetPath>& OutConstructClasses, const UClass* SpecificSubclass);
+#endif
 
 	/** Gets the root folder for UI folders for emitters. */
 	UNiagaraSystemEditorFolder& GetRootFolder() const;
diff --git a/Engine/Source/Developer/RigVMDeveloper/Private/RigVMModel/RigVMController.cpp b/Engine/Source/Developer/RigVMDeveloper/Private/RigVMModel/RigVMController.cpp
--- a/Engine/Source/Developer/RigVMDeveloper/Private/RigVMModel/RigVMController.cpp
+++ b/Engine/Source/Developer/RigVMDeveloper/Private/RigVMModel/RigVMController.cpp
@@ -107,6 +107,18 @@ URigVMController::~URigVMController()
 {
 }
 
+#if WITH_EDITORONLY_DATA
+void URigVMController::DeclareConstructClasses(TArray<FTopLevelAssetPath>& OutConstructClasses, const UClass* SpecificSubclass)
+{
+	Super::DeclareConstructClasses(OutConstructClasses, SpecificSubclass);
+	OutConstructClasses.Add(FTopLevelAssetPath(URigVMActionStack::StaticClass()));
+	OutConstructClasses.Add(FTopLevelAssetPath(URigVMInjectionInfo::StaticClass()));
+	OutConstructClasses.Add(FTopLevelAssetPath(URigVMPin::StaticClass()));
+	OutConstructClasses.Add(FTopLevelAssetPath(URigVMVariableNode::StaticClass()));
+	OutConstructClasses.Add(FTopLevelAssetPath(URigVMLink::StaticClass()));
+}
+#endif
+
 URigVMGraph* URigVMController::GetGraph() const
 {
 	if (Graphs.Num() == 0)
diff --git a/Engine/Source/Developer/RigVMDeveloper/Public/RigVMModel/RigVMController.h b/Engine/Source/Developer/RigVMDeveloper/Public/RigVMModel/RigVMController.h
--- a/Engine/Source/Developer/RigVMDeveloper/Public/RigVMModel/RigVMController.h
+++ b/Engine/Source/Developer/RigVMDeveloper/Public/RigVMModel/RigVMController.h
@@ -155,6 +155,10 @@ public:
 	// Default destructor
 	~URigVMController();
 
+#if WITH_EDITORONLY_DATA
+	static void DeclareConstructClasses(TArray<FTopLevelAssetPath>& OutConstructClasses, const UClass* SpecificSubclass);
+#endif
+
 	// Returns the currently edited Graph of this controller.
 	UFUNCTION(BlueprintCallable, Category = RigVMController)
 	URigVMGraph* GetGraph() const;
diff --git a/Engine/Source/Runtime/ClothingSystemRuntimeCommon/Private/ClothingAsset.cpp b/Engine/Source/Runtime/ClothingSystemRuntimeCommon/Private/ClothingAsset.cpp
--- a/Engine/Source/Runtime/ClothingSystemRuntimeCommon/Private/ClothingAsset.cpp
+++ b/Engine/Source/Runtime/ClothingSystemRuntimeCommon/Private/ClothingAsset.cpp
@@ -1308,6 +1308,29 @@ void UClothingAssetCommon::PostLoad()
 #endif  // #if WITH_EDITORONLY_DATA
 }
 
+#if WITH_EDITORONLY_DATA
+void UClothingAssetCommon::DeclareConstructClasses(TArray<FTopLevelAssetPath>& OutConstructClasses, const UClass* SpecificSubclass)
+{
+	Super::DeclareConstructClasses(OutConstructClasses, SpecificSubclass);
+	const TArray<IClothingSimulationFactoryClassProvider*> ClassProviders =
+		IModularFeatures::Get().GetModularFeatureImplementations<IClothingSimulationFactoryClassProvider>(IClothingSimulationFactoryClassProvider::FeatureName);
+	for (IClothingSimulationFactoryClassProvider* Provider : ClassProviders)
+	{
+		check(Provider);
+		if (UClass* const ClothingSimulationFactoryClass = *TSubclassOf<class UClothingSimulationFactory>(Provider->GetClothingSimulationFactoryClass()))
+		{
+			const UClothingSimulationFactory* const ClothingSimulationFactory = ClothingSimulationFactoryClass->GetDefaultObject<UClothingSimulationFactory>();
+			for (TSubclassOf<UClothConfigBase> ClothConfigClass : ClothingSimulationFactory->GetClothConfigClasses())
+			{
+				OutConstructClasses.Add(FTopLevelAssetPath(ClothConfigClass));
+			}
+		}
+	}
+	// OutConstructClasses.Add(FTopLevelAssetPath(UChaosClothSharedSimConfig::StaticClass()));
+	//OutConstructClasses.Add(FTopLevelAssetPath(UChaosClothConfig::StaticClass()));
+}
+#endif
+
 void UClothingAssetCommon::Serialize(FArchive& Ar)
 {
 	Super::Serialize(Ar);
diff --git a/Engine/Source/Runtime/ClothingSystemRuntimeCommon/Public/ClothingAsset.h b/Engine/Source/Runtime/ClothingSystemRuntimeCommon/Public/ClothingAsset.h
--- a/Engine/Source/Runtime/ClothingSystemRuntimeCommon/Public/ClothingAsset.h
+++ b/Engine/Source/Runtime/ClothingSystemRuntimeCommon/Public/ClothingAsset.h
@@ -230,6 +230,9 @@ PRAGMA_ENABLE_DEPRECATION_WARNINGS
 
 	/** Migrate deprecated objects. */
 	virtual void PostLoad() override;
+#if WITH_EDITORONLY_DATA
+	static void DeclareConstructClasses(TArray<FTopLevelAssetPath>& OutConstructClasses, const UClass* SpecificSubclass);
+#endif
 
 	/** Serialize deprecated objects. */
 	virtual void Serialize(FArchive& Ar) override;
diff --git a/Engine/Source/Runtime/Engine/Classes/Animation/AnimSequenceBase.h b/Engine/Source/Runtime/Engine/Classes/Animation/AnimSequenceBase.h
--- a/Engine/Source/Runtime/Engine/Classes/Animation/AnimSequenceBase.h
+++ b/Engine/Source/Runtime/Engine/Classes/Animation/AnimSequenceBase.h
@@ -76,6 +76,9 @@ public:
 	virtual void PostLoad() override;
 	virtual bool IsPostLoadThreadSafe() const override;
 	virtual void PostDuplicate(EDuplicateMode::Type DuplicateMode) override;
+#if WITH_EDITORONLY_DATA
+	static void DeclareConstructClasses(TArray<FTopLevelAssetPath>& OutConstructClasses, const UClass* SpecificSubclass);
+#endif
 	//~ End UObject Interface
 
 	/** Returns the total play length of the montage, if played back with a speed of 1.0. */
diff --git a/Engine/Source/Runtime/Engine/Classes/Engine/Blueprint.h b/Engine/Source/Runtime/Engine/Classes/Engine/Blueprint.h
--- a/Engine/Source/Runtime/Engine/Classes/Engine/Blueprint.h
+++ b/Engine/Source/Runtime/Engine/Classes/Engine/Blueprint.h
@@ -978,6 +978,9 @@ public:
 	virtual bool Rename(const TCHAR* NewName = nullptr, UObject* NewOuter = nullptr, ERenameFlags Flags = REN_None) override;
 	virtual UClass* RegenerateClass(UClass* ClassToRegenerate, UObject* PreviousCDO) override;
 	virtual void PostLoad() override;
+#if WITH_EDITORONLY_DATA
+	static void DeclareConstructClasses(TArray<FTopLevelAssetPath>& OutConstructClasses, const UClass* SpecificSubclass);
+#endif
 	virtual bool Modify(bool bAlwaysMarkDirty = true) override;
 	virtual void GetAssetRegistryTags(TArray<FAssetRegistryTag>& OutTags) const override;
 #if WITH_EDITOR
diff --git a/Engine/Source/Runtime/Engine/Classes/Engine/MapBuildDataRegistry.h b/Engine/Source/Runtime/Engine/Classes/Engine/MapBuildDataRegistry.h
--- a/Engine/Source/Runtime/Engine/Classes/Engine/MapBuildDataRegistry.h
+++ b/Engine/Source/Runtime/Engine/Classes/Engine/MapBuildDataRegistry.h
@@ -292,6 +292,10 @@ public:
 	//~ Begin UObject Interface
 	ENGINE_API virtual void Serialize(FArchive& Ar) override;
 	ENGINE_API virtual void PostLoad() override;
+#if WITH_EDITORONLY_DATA
+	ENGINE_API static void DeclareConstructClasses(TArray<FTopLevelAssetPath>& OutConstructClasses, const UClass* SpecificSubclass);
+#endif
+
 	static void AddReferencedObjects(UObject* InThis, FReferenceCollector& Collector);
 	ENGINE_API virtual void BeginDestroy() override;
 	ENGINE_API virtual bool IsReadyForFinishDestroy() override;
diff --git a/Engine/Source/Runtime/Engine/Classes/Engine/SkeletalMesh.h b/Engine/Source/Runtime/Engine/Classes/Engine/SkeletalMesh.h
--- a/Engine/Source/Runtime/Engine/Classes/Engine/SkeletalMesh.h
+++ b/Engine/Source/Runtime/Engine/Classes/Engine/SkeletalMesh.h
@@ -2428,6 +2428,10 @@ public:
 #endif
 	virtual void PostInitProperties() override;
 	virtual void PostLoad() override;
+#if WITH_EDITORONLY_DATA
+	static void DeclareConstructClasses(TArray<FTopLevelAssetPath>& OutConstructClasses, const UClass* SpecificSubclass);
+#endif
+
 	virtual bool IsPostLoadThreadSafe() const override;
 	virtual void GetAssetRegistryTags(TArray<FAssetRegistryTag>& OutTags) const override;
 	virtual FString GetDesc() override;
diff --git a/Engine/Source/Runtime/Engine/Classes/Engine/StaticMesh.h b/Engine/Source/Runtime/Engine/Classes/Engine/StaticMesh.h
--- a/Engine/Source/Runtime/Engine/Classes/Engine/StaticMesh.h
+++ b/Engine/Source/Runtime/Engine/Classes/Engine/StaticMesh.h
@@ -1610,6 +1610,9 @@ public:
 	ENGINE_API virtual void Serialize(FArchive& Ar) override;
 	ENGINE_API virtual void PostInitProperties() override;
 	ENGINE_API virtual void PostLoad() override;
+#if WITH_EDITORONLY_DATA
+	ENGINE_API static void DeclareConstructClasses(TArray<FTopLevelAssetPath>& OutConstructClasses, const UClass* SpecificSubclass);
+#endif
 	ENGINE_API virtual bool IsPostLoadThreadSafe() const override;
 	ENGINE_API virtual void BeginDestroy() override;
 	ENGINE_API virtual bool IsReadyForFinishDestroy() override;
diff --git a/Engine/Source/Runtime/Engine/Classes/Engine/World.h b/Engine/Source/Runtime/Engine/Classes/Engine/World.h
--- a/Engine/Source/Runtime/Engine/Classes/Engine/World.h
+++ b/Engine/Source/Runtime/Engine/Classes/Engine/World.h
@@ -2904,6 +2904,9 @@ public:
 	virtual void FinishDestroy() override;
 	virtual bool IsReadyForFinishDestroy() override;
 	virtual void PostLoad() override;
+#if WITH_EDITORONLY_DATA
+	static void DeclareConstructClasses(TArray<FTopLevelAssetPath>& OutConstructClasses, const UClass* SpecificSubclass);
+#endif
 	virtual void PreDuplicate(FObjectDuplicationParameters& DupParams) override;
 	PRAGMA_DISABLE_DEPRECATION_WARNINGS // Suppress compiler warning on override of deprecated function
 	UE_DEPRECATED(5.0, "Use version that takes FObjectPreSaveRootContext instead.")
diff --git a/Engine/Source/Runtime/Engine/Classes/GameFramework/WorldSettings.h b/Engine/Source/Runtime/Engine/Classes/GameFramework/WorldSettings.h
--- a/Engine/Source/Runtime/Engine/Classes/GameFramework/WorldSettings.h
+++ b/Engine/Source/Runtime/Engine/Classes/GameFramework/WorldSettings.h
@@ -812,6 +812,10 @@ protected:
 public:
 	//~ Begin UObject Interface.
 	virtual void PostLoad() override;
+#if WITH_EDITORONLY_DATA
+	static void DeclareConstructClasses(TArray<FTopLevelAssetPath>& OutConstructClasses, const UClass* SpecificSubclass);
+#endif
+
 #if WITH_EDITOR
 	virtual bool CanEditChange(const FProperty* InProperty) const override;
 	virtual void PostEditChangeProperty(FPropertyChangedEvent& PropertyChangedEvent) override;
diff --git a/Engine/Source/Runtime/Engine/Classes/Materials/Material.h b/Engine/Source/Runtime/Engine/Classes/Materials/Material.h
--- a/Engine/Source/Runtime/Engine/Classes/Materials/Material.h
+++ b/Engine/Source/Runtime/Engine/Classes/Materials/Material.h
@@ -1100,6 +1100,9 @@ public:
 	ENGINE_API virtual void Serialize(FArchive& Ar) override;
 	ENGINE_API virtual void PostDuplicate(bool bDuplicateForPIE) override;
 	ENGINE_API virtual void PostLoad() override;
+#if WITH_EDITORONLY_DATA
+	ENGINE_API static void DeclareConstructClasses(TArray<FTopLevelAssetPath>& OutConstructClasses, const UClass* SpecificSubclass);
+#endif
 #if WITH_EDITOR
 	ENGINE_API virtual void BeginCacheForCookedPlatformData( const ITargetPlatform *TargetPlatform ) override;
 	ENGINE_API virtual bool IsCachedCookedPlatformDataLoaded( const ITargetPlatform* TargetPlatform ) override;
diff --git a/Engine/Source/Runtime/Engine/Classes/Materials/MaterialFunction.h b/Engine/Source/Runtime/Engine/Classes/Materials/MaterialFunction.h
--- a/Engine/Source/Runtime/Engine/Classes/Materials/MaterialFunction.h
+++ b/Engine/Source/Runtime/Engine/Classes/Materials/MaterialFunction.h
@@ -108,6 +108,9 @@ public:
 #endif // WITH_EDITOR
 	virtual void Serialize(FArchive& Ar) override;
 	virtual void PostLoad() override;
+#if WITH_EDITORONLY_DATA
+	static void DeclareConstructClasses(TArray<FTopLevelAssetPath>& OutConstructClasses, const UClass* SpecificSubclass);
+#endif
 	//~ End UObject Interface.
 
 	void SetMaterialFunctionUsage(EMaterialFunctionUsage Usage) { MaterialFunctionUsage = Usage; }
diff --git a/Engine/Source/Runtime/Engine/Classes/Materials/MaterialFunctionInterface.h b/Engine/Source/Runtime/Engine/Classes/Materials/MaterialFunctionInterface.h
--- a/Engine/Source/Runtime/Engine/Classes/Materials/MaterialFunctionInterface.h
+++ b/Engine/Source/Runtime/Engine/Classes/Materials/MaterialFunctionInterface.h
@@ -63,6 +63,10 @@ public:
 	//~ Begin UObject Interface.
 	virtual void PostInitProperties() override;
 	virtual void PostLoad() override;
+#if WITH_EDITORONLY_DATA
+	ENGINE_API static void DeclareConstructClasses(TArray<FTopLevelAssetPath>& OutConstructClasses, const UClass* SpecificSubclass);
+#endif
+
 	virtual void GetAssetRegistryTags(TArray<FAssetRegistryTag>& OutTags) const override;
 #if WITH_EDITOR
 	virtual void PostEditChangeProperty(FPropertyChangedEvent& PropertyChangedEvent) override;
diff --git a/Engine/Source/Runtime/Engine/Classes/Materials/MaterialInstance.h b/Engine/Source/Runtime/Engine/Classes/Materials/MaterialInstance.h
--- a/Engine/Source/Runtime/Engine/Classes/Materials/MaterialInstance.h
+++ b/Engine/Source/Runtime/Engine/Classes/Materials/MaterialInstance.h
@@ -769,6 +769,9 @@ public:
 #endif
 	virtual ENGINE_API void Serialize(FArchive& Ar) override;
 	virtual ENGINE_API void PostLoad() override;
+#if WITH_EDITORONLY_DATA
+	ENGINE_API static void DeclareConstructClasses(TArray<FTopLevelAssetPath>& OutConstructClasses, const UClass* SpecificSubclass);
+#endif
 	virtual ENGINE_API void BeginDestroy() override;
 	virtual ENGINE_API bool IsReadyForFinishDestroy() override;
 	virtual ENGINE_API void FinishDestroy() override;
diff --git a/Engine/Source/Runtime/Engine/Classes/Materials/MaterialInterface.h b/Engine/Source/Runtime/Engine/Classes/Materials/MaterialInterface.h
--- a/Engine/Source/Runtime/Engine/Classes/Materials/MaterialInterface.h
+++ b/Engine/Source/Runtime/Engine/Classes/Materials/MaterialInterface.h
@@ -328,6 +328,10 @@ public:
 	ENGINE_API virtual void PostInitProperties() override;
 	ENGINE_API virtual void Serialize(FArchive& Ar) override;
 	ENGINE_API virtual void PostLoad() override;
+#if WITH_EDITORONLY_DATA
+	ENGINE_API static void DeclareConstructClasses(TArray<FTopLevelAssetPath>& OutConstructClasses, const UClass* SpecificSubclass);
+#endif
+
 	ENGINE_API virtual void PostDuplicate(bool bDuplicateForPIE) override;
 	ENGINE_API virtual void PostCDOContruct() override;
 	ENGINE_API virtual bool Rename(const TCHAR* NewName = nullptr, UObject* NewOuter = nullptr, ERenameFlags Flags = REN_None) override;
diff --git a/Engine/Source/Runtime/Engine/Classes/PhysicsEngine/PhysicsAsset.h b/Engine/Source/Runtime/Engine/Classes/PhysicsEngine/PhysicsAsset.h
--- a/Engine/Source/Runtime/Engine/Classes/PhysicsEngine/PhysicsAsset.h
+++ b/Engine/Source/Runtime/Engine/Classes/PhysicsEngine/PhysicsAsset.h
@@ -251,6 +251,9 @@ public:
 	//~ Begin UObject Interface
 	virtual void Serialize(FArchive& Ar) override;
 	virtual void PostLoad() override;
+#if WITH_EDITORONLY_DATA
+	static void DeclareConstructClasses(TArray<FTopLevelAssetPath>& OutConstructClasses, const UClass* SpecificSubclass);
+#endif
 	virtual FString GetDesc() override;
 	virtual void GetAssetRegistryTags(TArray<FAssetRegistryTag>& OutTags) const override;
 	virtual void GetResourceSizeEx(FResourceSizeEx& CumulativeResourceSize) override;
diff --git a/Engine/Source/Runtime/Engine/Private/Animation/AnimSequenceBase.cpp b/Engine/Source/Runtime/Engine/Private/Animation/AnimSequenceBase.cpp
--- a/Engine/Source/Runtime/Engine/Private/Animation/AnimSequenceBase.cpp
+++ b/Engine/Source/Runtime/Engine/Private/Animation/AnimSequenceBase.cpp
@@ -57,6 +57,29 @@ bool UAnimSequenceBase::IsPostLoadThreadSafe() const
 	return false;	// PostLoad is not thread safe because of the call to VerifyCurveNames() (calling USkeleton::VerifySmartName) that can mutate a shared map in the skeleton.
 }
 
+#if WITH_EDITORONLY_DATA
+void UAnimSequenceBase::DeclareConstructClasses(TArray<FTopLevelAssetPath>& OutConstructClasses, const UClass* SpecificSubclass)
+{
+	Super::DeclareConstructClasses(OutConstructClasses, SpecificSubclass);
+
+	OutConstructClasses.Add(FTopLevelAssetPath(TEXT("/Script/Engine.AnimDataModel")));
+
+	// We need to declare all types that can be returned from UE::Anim::DataModel::IAnimationDataModels::FindClassForAnimationAsset(UAnimSequenceBase*);
+	OutConstructClasses.Add(FTopLevelAssetPath(TEXT("/Script/AnimationData.AnimationSequencerDataModel")));
+
+	// We can call Controller->CreateModel, which can add objects, so add every ControllerClass as something we can construct.
+	// THe caller will add on all recursively constructable classes
+	UClass* AnimationControllerClass = UAnimationDataController::StaticClass();
+	for (TObjectIterator<UClass> Iter; Iter; ++Iter)
+	{
+		if ((*Iter)->ImplementsInterface(AnimationControllerClass))
+		{
+			OutConstructClasses.Add(FTopLevelAssetPath(*Iter));
+		}
+	}
+}
+#endif
+
 void UAnimSequenceBase::PostLoad()
 {
 #if WITH_EDITORONLY_DATA
diff --git a/Engine/Source/Runtime/Engine/Private/Blueprint.cpp b/Engine/Source/Runtime/Engine/Private/Blueprint.cpp
--- a/Engine/Source/Runtime/Engine/Private/Blueprint.cpp
+++ b/Engine/Source/Runtime/Engine/Private/Blueprint.cpp
@@ -780,6 +780,15 @@ void UBlueprint::PostLoad()
 #endif
 }
 
+#if WITH_EDITORONLY_DATA
+void UBlueprint::DeclareConstructClasses(TArray<FTopLevelAssetPath>& OutConstructClasses, const UClass* SpecificSubclass)
+{
+	Super::DeclareConstructClasses(OutConstructClasses, SpecificSubclass);
+	OutConstructClasses.Add(FTopLevelAssetPath(UObjectRedirector::StaticClass()));
+}
+#endif
+
+
 void UBlueprint::DebuggingWorldRegistrationHelper(UObject* ObjectProvidingWorld, UObject* ValueToRegister)
 {
 	if (ObjectProvidingWorld != NULL)
diff --git a/Engine/Source/Runtime/Engine/Private/MapBuildData.cpp b/Engine/Source/Runtime/Engine/Private/MapBuildData.cpp
--- a/Engine/Source/Runtime/Engine/Private/MapBuildData.cpp
+++ b/Engine/Source/Runtime/Engine/Private/MapBuildData.cpp
@@ -585,6 +585,13 @@ void UMapBuildDataRegistry::PostLoad()
 	SetupLightmapResourceClusters();
 }
 
+#if WITH_EDITORONLY_DATA
+void UMapBuildDataRegistry::DeclareConstructClasses(TArray<FTopLevelAssetPath>& OutConstructClasses, const UClass* SpecificSubclass)
+{
+	Super::DeclareConstructClasses(OutConstructClasses, SpecificSubclass);
+	OutConstructClasses.Add(FTopLevelAssetPath(UTextureCube::StaticClass()));
+}
+#endif
 
 void UMapBuildDataRegistry::HandleLegacyEncodedCubemapData()
 {
diff --git a/Engine/Source/Runtime/Engine/Private/Materials/Material.cpp b/Engine/Source/Runtime/Engine/Private/Materials/Material.cpp
--- a/Engine/Source/Runtime/Engine/Private/Materials/Material.cpp
+++ b/Engine/Source/Runtime/Engine/Private/Materials/Material.cpp
@@ -3846,6 +3846,14 @@ void UMaterial::PostLoad()
 	//DumpDebugInfo(*GLog);
 }
 
+#if WITH_EDITORONLY_DATA
+void UMaterial::DeclareConstructClasses(TArray<FTopLevelAssetPath>& OutConstructClasses, const UClass* SpecificSubclass)
+{
+	Super::DeclareConstructClasses(OutConstructClasses, SpecificSubclass);
+	OutConstructClasses.Add(FTopLevelAssetPath(UMaterialEditorOnlyData::StaticClass()));
+}
+#endif
+
 void UMaterial::DumpDebugInfo(FOutputDevice& OutputDevice) const
 {
 	for (FMaterialResource* Resource : MaterialResources)
diff --git a/Engine/Source/Runtime/Engine/Private/Materials/MaterialExpressions.cpp b/Engine/Source/Runtime/Engine/Private/Materials/MaterialExpressions.cpp
--- a/Engine/Source/Runtime/Engine/Private/Materials/MaterialExpressions.cpp
+++ b/Engine/Source/Runtime/Engine/Private/Materials/MaterialExpressions.cpp
@@ -13344,6 +13344,14 @@ void UMaterialFunctionInterface::PostLoad()
 	}
 }
 
+#if WITH_EDITORONLY_DATA
+void UMaterialFunctionInterface::DeclareConstructClasses(TArray<FTopLevelAssetPath>& OutConstructClasses, const UClass* SpecificSubclass)
+{
+	Super::DeclareConstructClasses(OutConstructClasses, SpecificSubclass);
+	OutConstructClasses.Add(FTopLevelAssetPath(UMaterialFunctionInterfaceEditorOnlyData::StaticClass()));
+}
+#endif
+
 void UMaterialFunctionInterface::GetAssetRegistryTags(TArray<FAssetRegistryTag>& OutTags) const
 {
 	Super::GetAssetRegistryTags(OutTags);
@@ -13929,6 +13937,15 @@ void UMaterialFunction::PostLoad()
 #endif // #if WITH_EDITOR
 }
 
+#if WITH_EDITORONLY_DATA
+void UMaterialFunction::DeclareConstructClasses(TArray<FTopLevelAssetPath>& OutConstructClasses, const UClass* SpecificSubclass)
+{
+	Super::DeclareConstructClasses(OutConstructClasses, SpecificSubclass);
+	OutConstructClasses.Add(FTopLevelAssetPath(UMaterialFunctionEditorOnlyData::StaticClass()));
+}
+#endif
+
+
 #if WITH_EDITORONLY_DATA
 TConstArrayView<TObjectPtr<UMaterialExpression>> UMaterialFunction::GetExpressions() const
 {
diff --git a/Engine/Source/Runtime/Engine/Private/Materials/MaterialInstance.cpp b/Engine/Source/Runtime/Engine/Private/Materials/MaterialInstance.cpp
--- a/Engine/Source/Runtime/Engine/Private/Materials/MaterialInstance.cpp
+++ b/Engine/Source/Runtime/Engine/Private/Materials/MaterialInstance.cpp
@@ -2954,6 +2954,14 @@ void UMaterialInstance::PostLoad()
 	//DumpDebugInfo(*GLog);
 }
 
+#if WITH_EDITORONLY_DATA
+void UMaterialInstance::DeclareConstructClasses(TArray<FTopLevelAssetPath>& OutConstructClasses, const UClass* SpecificSubclass)
+{
+	Super::DeclareConstructClasses(OutConstructClasses, SpecificSubclass);
+	OutConstructClasses.Add(FTopLevelAssetPath(UMaterialInstanceEditorOnlyData::StaticClass()));
+}
+#endif
+
 void UMaterialInstance::BeginDestroy()
 {
 	TArray<TRefCountPtr<FMaterialResource>> ResourcesToDestroy;
diff --git a/Engine/Source/Runtime/Engine/Private/Materials/MaterialInterface.cpp b/Engine/Source/Runtime/Engine/Private/Materials/MaterialInterface.cpp
--- a/Engine/Source/Runtime/Engine/Private/Materials/MaterialInterface.cpp
+++ b/Engine/Source/Runtime/Engine/Private/Materials/MaterialInterface.cpp
@@ -207,6 +207,14 @@ void UMaterialInterface::PostLoad()
 #endif // WITH_EDITORONLY_DATA
 }
 
+#if WITH_EDITORONLY_DATA
+void UMaterialInterface::DeclareConstructClasses(TArray<FTopLevelAssetPath>& OutConstructClasses, const UClass* SpecificSubclass)
+{
+	Super::DeclareConstructClasses(OutConstructClasses, SpecificSubclass);
+	OutConstructClasses.Add(FTopLevelAssetPath(UMaterialInterfaceEditorOnlyData::StaticClass()));
+}
+#endif
+
 const FMaterialCachedExpressionData& UMaterialInterface::GetCachedExpressionData(TMicRecursionGuard) const
 {
 	const FMaterialCachedExpressionData* LocalData = CachedExpressionData.Get();
diff --git a/Engine/Source/Runtime/Engine/Private/PhysicsEngine/PhysicsAsset.cpp b/Engine/Source/Runtime/Engine/Private/PhysicsEngine/PhysicsAsset.cpp
--- a/Engine/Source/Runtime/Engine/Private/PhysicsEngine/PhysicsAsset.cpp
+++ b/Engine/Source/Runtime/Engine/Private/PhysicsEngine/PhysicsAsset.cpp
@@ -209,6 +209,16 @@ void UPhysicsAsset::PostLoad()
 	}
 }
 
+#if WITH_EDITORONLY_DATA
+void UPhysicsAsset::DeclareConstructClasses(TArray<FTopLevelAssetPath>& OutConstructClasses, const UClass* SpecificSubclass)
+{
+	Super::DeclareConstructClasses(OutConstructClasses, SpecificSubclass);
+
+	OutConstructClasses.Add(FTopLevelAssetPath(USkeletalBodySetup::StaticClass()));
+}
+#endif
+
+
 void UPhysicsAsset::Serialize(FArchive& Ar)
 {
 	Ar.UsingCustomVersion(FFortniteNCBranchObjectVersion::GUID);
diff --git a/Engine/Source/Runtime/Engine/Private/SkeletalMesh.cpp b/Engine/Source/Runtime/Engine/Private/SkeletalMesh.cpp
--- a/Engine/Source/Runtime/Engine/Private/SkeletalMesh.cpp
+++ b/Engine/Source/Runtime/Engine/Private/SkeletalMesh.cpp
@@ -3073,6 +3073,15 @@ void USkeletalMesh::BeginPostLoadInternal(FSkinnedAssetPostLoadContext& Context)
 #endif // #if WITH_EDITOR
 }
 
+#if WITH_EDITORONLY_DATA
+void USkeletalMesh::DeclareConstructClasses(TArray<FTopLevelAssetPath>& OutConstructClasses, const UClass* SpecificSubclass)
+{
+	Super::DeclareConstructClasses(OutConstructClasses, SpecificSubclass);
+
+	OutConstructClasses.Add(FTopLevelAssetPath(USkeletalMeshEditorData::StaticClass()));
+}
+#endif
+
 void USkeletalMesh::PostLoad()
 {
 	LLM_SCOPE(ELLMTag::SkeletalMesh);
diff --git a/Engine/Source/Runtime/Engine/Private/StaticMesh.cpp b/Engine/Source/Runtime/Engine/Private/StaticMesh.cpp
--- a/Engine/Source/Runtime/Engine/Private/StaticMesh.cpp
+++ b/Engine/Source/Runtime/Engine/Private/StaticMesh.cpp
@@ -5718,6 +5718,15 @@ void UStaticMesh::PostLoad()
 	}
 }
 
+#if WITH_EDITORONLY_DATA
+void UStaticMesh::DeclareConstructClasses(TArray<FTopLevelAssetPath>& OutConstructClasses, const UClass* SpecificSubclass)
+{
+	Super::DeclareConstructClasses(OutConstructClasses, SpecificSubclass);
+	OutConstructClasses.Add(FTopLevelAssetPath(UStaticMeshDescriptionBulkData::StaticClass()));
+	OutConstructClasses.Add(FTopLevelAssetPath(TEXT("/Script/NavigationSystem.NavCollision")));
+}
+#endif
+
 void UStaticMesh::BeginPostLoadInternal(FStaticMeshPostLoadContext& Context)
 {
 #if WITH_EDITOR
diff --git a/Engine/Source/Runtime/Engine/Private/World.cpp b/Engine/Source/Runtime/Engine/Private/World.cpp
--- a/Engine/Source/Runtime/Engine/Private/World.cpp
+++ b/Engine/Source/Runtime/Engine/Private/World.cpp
@@ -1284,6 +1284,15 @@ void UWorld::PostLoad()
 	ResetAverageRequiredTexturePoolSize();
 }
 
+#if WITH_EDITORONLY_DATA
+void UWorld::DeclareConstructClasses(TArray<FTopLevelAssetPath>& OutConstructClasses, const UClass* SpecificSubclass)
+{
+	Super::DeclareConstructClasses(OutConstructClasses, SpecificSubclass);
+	OutConstructClasses.Add(FTopLevelAssetPath(GEngine->WorldSettingsClass));
+	OutConstructClasses.Add(FTopLevelAssetPath(UWorldThumbnailInfo::StaticClass()));
+}
+#endif
+
 void UWorld::PreDuplicate(FObjectDuplicationParameters& DupParams)
 {
 	if (PersistentLevel)
diff --git a/Engine/Source/Runtime/Engine/Private/WorldSettings.cpp b/Engine/Source/Runtime/Engine/Private/WorldSettings.cpp
--- a/Engine/Source/Runtime/Engine/Private/WorldSettings.cpp
+++ b/Engine/Source/Runtime/Engine/Private/WorldSettings.cpp
@@ -565,6 +565,19 @@ PRAGMA_DISABLE_DEPRECATION_WARNINGS
 PRAGMA_ENABLE_DEPRECATION_WARNINGS
 }
 
+#if WITH_EDITORONLY_DATA
+void AWorldSettings::DeclareConstructClasses(TArray<FTopLevelAssetPath>& OutConstructClasses, const UClass* SpecificSubclass)
+{
+	Super::DeclareConstructClasses(OutConstructClasses, SpecificSubclass);
+	OutConstructClasses.Add(FTopLevelAssetPath(UNavigationSystemConfig::StaticClass()));	
+	TSubclassOf<UNavigationSystemConfig> NavSystemConfigClass = UNavigationSystemConfig::GetDefaultConfigClass();
+	if (*NavSystemConfigClass)
+	{
+		OutConstructClasses.Add(FTopLevelAssetPath(NavSystemConfigClass));
+	}
+}
+#endif
+
 bool AWorldSettings::IsNavigationSystemEnabled() const
 {
 	return NavigationSystemConfig && NavigationSystemConfig->NavigationSystemClass.IsValid();
diff --git a/Engine/Source/Runtime/Landscape/Classes/LandscapeComponent.h b/Engine/Source/Runtime/Landscape/Classes/LandscapeComponent.h
--- a/Engine/Source/Runtime/Landscape/Classes/LandscapeComponent.h
+++ b/Engine/Source/Runtime/Landscape/Classes/LandscapeComponent.h
@@ -644,6 +644,10 @@ public:
 	virtual void BeginDestroy() override;
 	virtual void PostDuplicate(bool bDuplicateForPIE) override;
 	virtual void PostLoad() override;
+#if WITH_EDITORONLY_DATA
+	static void DeclareConstructClasses(TArray<FTopLevelAssetPath>& OutConstructClasses, const UClass* SpecificSubclass);
+#endif
+
 #if WITH_EDITOR
 	virtual void BeginCacheForCookedPlatformData(const ITargetPlatform* TargetPlatform) override;
 	virtual void PostEditUndo() override;
diff --git a/Engine/Source/Runtime/Landscape/Private/Landscape.cpp b/Engine/Source/Runtime/Landscape/Private/Landscape.cpp
--- a/Engine/Source/Runtime/Landscape/Private/Landscape.cpp
+++ b/Engine/Source/Runtime/Landscape/Private/Landscape.cpp
@@ -1211,6 +1211,14 @@ void ULandscapeComponent::PostLoad()
 
 }
 
+#if WITH_EDITORONLY_DATA
+void ULandscapeComponent::DeclareConstructClasses(TArray<FTopLevelAssetPath>& OutConstructClasses, const UClass* SpecificSubclass)
+{
+	Super::DeclareConstructClasses(OutConstructClasses, SpecificSubclass);
+	OutConstructClasses.Add(FTopLevelAssetPath(ULandscapeMaterialInstanceConstant::StaticClass()));
+}
+#endif
+
 #if WITH_EDITORONLY_DATA
 TArray<ALandscapeProxy*> ALandscapeProxy::LandscapeProxies;
 #endif
diff --git a/Engine/Source/Runtime/LevelSequence/Private/LevelSequence.cpp b/Engine/Source/Runtime/LevelSequence/Private/LevelSequence.cpp
--- a/Engine/Source/Runtime/LevelSequence/Private/LevelSequence.cpp
+++ b/Engine/Source/Runtime/LevelSequence/Private/LevelSequence.cpp
@@ -333,6 +333,14 @@ void ULevelSequence::PostLoad()
 #endif
 }
 
+#if WITH_EDITORONLY_DATA
+void ULevelSequence::DeclareConstructClasses(TArray<FTopLevelAssetPath>& OutConstructClasses, const UClass* SpecificSubclass)
+{
+	Super::DeclareConstructClasses(OutConstructClasses, SpecificSubclass);
+	OutConstructClasses.Add(FTopLevelAssetPath(UObjectRedirector::StaticClass()));
+}
+#endif
+
 void ULevelSequence::PostInitProperties()
 {
 	Super::PostInitProperties();
diff --git a/Engine/Source/Runtime/LevelSequence/Private/LevelSequenceActor.cpp b/Engine/Source/Runtime/LevelSequence/Private/LevelSequenceActor.cpp
--- a/Engine/Source/Runtime/LevelSequence/Private/LevelSequenceActor.cpp
+++ b/Engine/Source/Runtime/LevelSequence/Private/LevelSequenceActor.cpp
@@ -246,6 +246,14 @@ void ALevelSequenceActor::PostLoad()
 #endif
 }
 
+#if WITH_EDITORONLY_DATA
+void ALevelSequenceActor::DeclareConstructClasses(TArray<FTopLevelAssetPath>& OutConstructClasses, const UClass* SpecificSubclass)
+{
+	Super::DeclareConstructClasses(OutConstructClasses, SpecificSubclass);
+	OutConstructClasses.Add(FTopLevelAssetPath(UDefaultLevelSequenceInstanceData::StaticClass()));
+}
+#endif
+
 ULevelSequence* ALevelSequenceActor::GetSequence() const
 {
 	return LevelSequenceAsset;
diff --git a/Engine/Source/Runtime/LevelSequence/Public/LevelSequence.h b/Engine/Source/Runtime/LevelSequence/Public/LevelSequence.h
--- a/Engine/Source/Runtime/LevelSequence/Public/LevelSequence.h
+++ b/Engine/Source/Runtime/LevelSequence/Public/LevelSequence.h
@@ -59,6 +59,10 @@ public:
 	virtual bool CanAnimateObject(UObject& InObject) const override;
 	virtual UObject* CreateDirectorInstance(IMovieScenePlayer& Player, FMovieSceneSequenceID SequenceID) override;
 	virtual void PostLoad() override;
+#if WITH_EDITORONLY_DATA
+	static void DeclareConstructClasses(TArray<FTopLevelAssetPath>& OutConstructClasses, const UClass* SpecificSubclass);
+#endif
+
 	virtual void PostInitProperties() override;
 	virtual bool Rename(const TCHAR* NewName = nullptr, UObject* NewOuter = nullptr, ERenameFlags Flags = REN_None) override;
 
diff --git a/Engine/Source/Runtime/LevelSequence/Public/LevelSequenceActor.h b/Engine/Source/Runtime/LevelSequence/Public/LevelSequenceActor.h
--- a/Engine/Source/Runtime/LevelSequence/Public/LevelSequenceActor.h
+++ b/Engine/Source/Runtime/LevelSequence/Public/LevelSequenceActor.h
@@ -268,6 +268,12 @@ protected:
 	virtual bool ReplicateSubobjects(UActorChannel* Channel, FOutBunch* Bunch, FReplicationFlags *RepFlags) override;
 	virtual void PostInitProperties() override;
 	virtual void PostLoad() override;
+public:
+#if WITH_EDITORONLY_DATA
+	static void DeclareConstructClasses(TArray<FTopLevelAssetPath>& OutConstructClasses, const UClass* SpecificSubclass);
+#endif
+protected:
+
 	//~ End UObject interface
 
 	//~ Begin AActor interface
diff --git a/Engine/Source/Runtime/MovieScene/Private/MovieScene.cpp b/Engine/Source/Runtime/MovieScene/Private/MovieScene.cpp
--- a/Engine/Source/Runtime/MovieScene/Private/MovieScene.cpp
+++ b/Engine/Source/Runtime/MovieScene/Private/MovieScene.cpp
@@ -104,6 +104,14 @@ void UMovieScene::PostLoad()
 	Super::PostLoad();
 }
 
+#if WITH_EDITORONLY_DATA
+void UMovieScene::DeclareConstructClasses(TArray<FTopLevelAssetPath>& OutConstructClasses, const UClass* SpecificSubclass)
+{
+	Super::DeclareConstructClasses(OutConstructClasses, SpecificSubclass);
+	OutConstructClasses.Add(FTopLevelAssetPath(TEXT("/Script/MovieScene.MovieSceneNodeGroupCollection")));
+}
+#endif
+
 void UMovieScene::Serialize( FArchive& Ar )
 {
 	Ar.UsingCustomVersion(FMovieSceneEvaluationCustomVersion::GUID);
diff --git a/Engine/Source/Runtime/MovieScene/Private/MovieSceneSection.cpp b/Engine/Source/Runtime/MovieScene/Private/MovieSceneSection.cpp
--- a/Engine/Source/Runtime/MovieScene/Private/MovieSceneSection.cpp
+++ b/Engine/Source/Runtime/MovieScene/Private/MovieSceneSection.cpp
@@ -121,6 +121,15 @@ void UMovieSceneSection::Serialize(FArchive& Ar)
 	}
 }
 
+#if WITH_EDITORONLY_DATA
+void UMovieSceneSection::DeclareConstructClasses(TArray<FTopLevelAssetPath>& OutConstructClasses, const UClass* SpecificSubclass)
+{
+	Super::DeclareConstructClasses(OutConstructClasses, SpecificSubclass);
+	OutConstructClasses.Add(FTopLevelAssetPath(UMovieSceneBuiltInEasingFunction::StaticClass()));
+}
+#endif
+
+
 void UMovieSceneSection::PostDuplicate(bool bDuplicateForPIE)
 {
 	if (ChannelProxyType == EMovieSceneChannelProxyType::Dynamic)
diff --git a/Engine/Source/Runtime/MovieScene/Public/MovieScene.h b/Engine/Source/Runtime/MovieScene/Public/MovieScene.h
--- a/Engine/Source/Runtime/MovieScene/Public/MovieScene.h
+++ b/Engine/Source/Runtime/MovieScene/Public/MovieScene.h
@@ -380,6 +380,10 @@ public:
 	virtual bool IsPostLoadThreadSafe() const override;
 	virtual void PostInitProperties() override;
 	virtual void PostLoad() override;
+#if WITH_EDITORONLY_DATA
+	static void DeclareConstructClasses(TArray<FTopLevelAssetPath>& OutConstructClasses, const UClass* SpecificSubclass);
+#endif
+
 
 #if WITH_EDITOR
 	virtual void PostEditUndo() override;
diff --git a/Engine/Source/Runtime/MovieScene/Public/MovieSceneSection.h b/Engine/Source/Runtime/MovieScene/Public/MovieSceneSection.h
--- a/Engine/Source/Runtime/MovieScene/Public/MovieSceneSection.h
+++ b/Engine/Source/Runtime/MovieScene/Public/MovieSceneSection.h
@@ -684,6 +684,10 @@ public:
 
 	MOVIESCENE_API void BuildDefaultComponents(UMovieSceneEntitySystemLinker* EntityLinker, const UE::MovieScene::FEntityImportParams& Params, UE::MovieScene::FImportedEntity* OutLedgerEntry);
 
+#if WITH_EDITORONLY_DATA
+	MOVIESCENE_API static void DeclareConstructClasses(TArray<FTopLevelAssetPath>& OutConstructClasses, const UClass* SpecificSubclass);
+#endif
+
 protected:
 
 	//~ UObject interface
diff --git a/Engine/Source/Runtime/RigVM/Private/RigVMCore/RigVM.cpp b/Engine/Source/Runtime/RigVM/Private/RigVMCore/RigVM.cpp
--- a/Engine/Source/Runtime/RigVM/Private/RigVMCore/RigVM.cpp
+++ b/Engine/Source/Runtime/RigVM/Private/RigVMCore/RigVM.cpp
@@ -259,6 +259,14 @@ void URigVM::PostLoad()
 	}
 }
 
+#if WITH_EDITORONLY_DATA
+void URigVM::DeclareConstructClasses(TArray<FTopLevelAssetPath>& OutConstructClasses, const UClass* SpecificSubclass)
+{
+	Super::DeclareConstructClasses(OutConstructClasses, SpecificSubclass);
+	OutConstructClasses.Add(FTopLevelAssetPath(URigVMMemoryStorage::StaticClass()));
+}
+#endif
+
 uint32 URigVM::GetVMHash() const
 {
 	uint32 Hash = 0;
diff --git a/Engine/Source/Runtime/RigVM/Public/RigVMCore/RigVM.h b/Engine/Source/Runtime/RigVM/Public/RigVMCore/RigVM.h
--- a/Engine/Source/Runtime/RigVM/Public/RigVMCore/RigVM.h
+++ b/Engine/Source/Runtime/RigVM/Public/RigVMCore/RigVM.h
@@ -180,6 +180,10 @@ public:
 	virtual void Save(FArchive& Ar);
 	virtual void Load(FArchive& Ar);
 	virtual void PostLoad() override;
+#if WITH_EDITORONLY_DATA
+	static void DeclareConstructClasses(TArray<FTopLevelAssetPath>& OutConstructClasses, const UClass* SpecificSubclass);
+#endif
+
 
 	// returns true if this is a nativized VM
 	virtual bool IsNativized() const { return false; }
