Volumetric Cloud now uses mode 2 in order to test how it would look and cost. Should fix reprojection and opaque edge artefact.

diff --git a/Engine/Config/BaseScalability.ini b/Engine/Config/BaseScalability.ini
--- a/Engine/Config/BaseScalability.ini
+++ b/Engine/Config/BaseScalability.ini
@@ -604,6 +604,7 @@ r.SkyAtmosphere.TransmittanceLUT.SampleCount=10.0
 r.SkyAtmosphere.MultiScatteringLUT.SampleCount=15.0
 r.SkyLight.RealTimeReflectionCapture=0
 r.VolumetricCloud.ShadowMap.SpatialFiltering=0
+r.VolumetricRenderTarget.Mode=2
 fx.Niagara.QualityLevel=0
 
 [EffectsQuality@1]
@@ -635,6 +636,7 @@ r.SkyAtmosphere.TransmittanceLUT.SampleCount=10.0
 r.SkyAtmosphere.MultiScatteringLUT.SampleCount=15.0
 r.SkyLight.RealTimeReflectionCapture=0
 r.VolumetricCloud.ShadowMap.SpatialFiltering=0
+r.VolumetricRenderTarget.Mode=2
 fx.Niagara.QualityLevel=1
 
 [EffectsQuality@2]
@@ -666,6 +668,7 @@ r.SkyAtmosphere.TransmittanceLUT.SampleCount=10.0
 r.SkyAtmosphere.MultiScatteringLUT.SampleCount=15.0
 r.SkyLight.RealTimeReflectionCapture=1
 r.VolumetricCloud.ShadowMap.SpatialFiltering=0
+r.VolumetricRenderTarget.Mode=2
 fx.Niagara.QualityLevel=2
 
 [EffectsQuality@3]
@@ -697,6 +700,7 @@ r.SkyAtmosphere.TransmittanceLUT.SampleCount=10.0
 r.SkyAtmosphere.MultiScatteringLUT.SampleCount=15.0
 r.SkyLight.RealTimeReflectionCapture=1
 r.VolumetricCloud.ShadowMap.SpatialFiltering=0
+r.VolumetricRenderTarget.Mode=2
 fx.Niagara.QualityLevel=3
 
 [EffectsQuality@Cine]
@@ -728,6 +732,7 @@ r.SkyAtmosphere.TransmittanceLUT.SampleCount=30.0
 r.SkyAtmosphere.MultiScatteringLUT.SampleCount=20.0
 r.SkyLight.RealTimeReflectionCapture=1
 r.VolumetricCloud.ShadowMap.SpatialFiltering=1
+r.VolumetricRenderTarget.Mode=2
 fx.Niagara.QualityLevel=4
 
 ;-----------------------------------------------------------------------------------------------------------------
