Clouds are now using default mode0 again to go back to original performance.
Game still have the shadow spatial filtering disabled in its own config. (not in base scalability, which will fix UE-167200).
Tested in editor

diff --git a/Engine/Config/BaseScalability.ini b/Engine/Config/BaseScalability.ini
--- a/Engine/Config/BaseScalability.ini
+++ b/Engine/Config/BaseScalability.ini
@@ -603,8 +603,6 @@ r.SkyAtmosphere.TransmittanceLUT.UseSmallFormat=1
 r.SkyAtmosphere.TransmittanceLUT.SampleCount=10.0
 r.SkyAtmosphere.MultiScatteringLUT.SampleCount=15.0
 r.SkyLight.RealTimeReflectionCapture=0
-r.VolumetricCloud.ShadowMap.SpatialFiltering=0
-r.VolumetricRenderTarget.Mode=2
 fx.Niagara.QualityLevel=0
 
 [EffectsQuality@1]
@@ -635,8 +633,6 @@ r.SkyAtmosphere.TransmittanceLUT.UseSmallFormat=0
 r.SkyAtmosphere.TransmittanceLUT.SampleCount=10.0
 r.SkyAtmosphere.MultiScatteringLUT.SampleCount=15.0
 r.SkyLight.RealTimeReflectionCapture=0
-r.VolumetricCloud.ShadowMap.SpatialFiltering=0
-r.VolumetricRenderTarget.Mode=2
 fx.Niagara.QualityLevel=1
 
 [EffectsQuality@2]
@@ -667,8 +663,6 @@ r.SkyAtmosphere.TransmittanceLUT.UseSmallFormat=0
 r.SkyAtmosphere.TransmittanceLUT.SampleCount=10.0
 r.SkyAtmosphere.MultiScatteringLUT.SampleCount=15.0
 r.SkyLight.RealTimeReflectionCapture=1
-r.VolumetricCloud.ShadowMap.SpatialFiltering=0
-r.VolumetricRenderTarget.Mode=2
 fx.Niagara.QualityLevel=2
 
 [EffectsQuality@3]
@@ -699,8 +693,6 @@ r.SkyAtmosphere.TransmittanceLUT.UseSmallFormat=0
 r.SkyAtmosphere.TransmittanceLUT.SampleCount=10.0
 r.SkyAtmosphere.MultiScatteringLUT.SampleCount=15.0
 r.SkyLight.RealTimeReflectionCapture=1
-r.VolumetricCloud.ShadowMap.SpatialFiltering=0
-r.VolumetricRenderTarget.Mode=2
 fx.Niagara.QualityLevel=3
 
 [EffectsQuality@Cine]
@@ -731,8 +723,6 @@ r.SkyAtmosphere.TransmittanceLUT.UseSmallFormat=0
 r.SkyAtmosphere.TransmittanceLUT.SampleCount=30.0
 r.SkyAtmosphere.MultiScatteringLUT.SampleCount=20.0
 r.SkyLight.RealTimeReflectionCapture=1
-r.VolumetricCloud.ShadowMap.SpatialFiltering=1
-r.VolumetricRenderTarget.Mode=2
 fx.Niagara.QualityLevel=4
 
 ;-----------------------------------------------------------------------------------------------------------------
