Update iOS devices to include new iPads, iPhone14s, AppleTV and new iPhoneSE

#rnx

diff --git a/Engine/Config/BaseDeviceProfiles.ini b/Engine/Config/BaseDeviceProfiles.ini
--- a/Engine/Config/BaseDeviceProfiles.ini
+++ b/Engine/Config/BaseDeviceProfiles.ini
@@ -12,6 +12,7 @@
 +DeviceProfileNameAndTypes=iPadAir2,IOS
 +DeviceProfileNameAndTypes=iPadAir3,IOS
 +DeviceProfileNameAndTypes=iPadAir4,IOS
++DeviceProfileNameAndTypes=iPadAir5,IOS
 +DeviceProfileNameAndTypes=iPadMini4,IOS
 +DeviceProfileNameAndTypes=iPadMini5,IOS
 +DeviceProfileNameAndTypes=iPadMini6,IOS
@@ -39,6 +40,11 @@
 +DeviceProfileNameAndTypes=iPhone13,IOS
 +DeviceProfileNameAndTypes=iPhone13Pro,IOS
 +DeviceProfileNameAndTypes=iPhone13ProMax,IOS
++DeviceProfileNameAndTypes=iPhoneSE3,IOS
++DeviceProfileNameAndTypes=iPhone14,IOS
++DeviceProfileNameAndTypes=iPhone14Plus,IOS
++DeviceProfileNameAndTypes=iPhone14Pro,IOS
++DeviceProfileNameAndTypes=iPhone14ProMax,IOS
 +DeviceProfileNameAndTypes=iPadPro,IOS
 +DeviceProfileNameAndTypes=iPadPro105,IOS
 +DeviceProfileNameAndTypes=iPadPro129,IOS
@@ -49,14 +55,18 @@
 +DeviceProfileNameAndTypes=iPad7,IOS
 +DeviceProfileNameAndTypes=iPad8,IOS
 +DeviceProfileNameAndTypes=iPad9,IOS
++DeviceProfileNameAndTypes=iPad10,IOS
 +DeviceProfileNameAndTypes=iPadPro11,IOS
 +DeviceProfileNameAndTypes=iPadPro2_11,IOS
 +DeviceProfileNameAndTypes=iPadPro3_11,IOS
++DeviceProfileNameAndTypes=iPadPro4_11,IOS
 +DeviceProfileNameAndTypes=iPadPro3_129,IOS
 +DeviceProfileNameAndTypes=iPadPro4_129,IOS
 +DeviceProfileNameAndTypes=iPadPro5_129,IOS
++DeviceProfileNameAndTypes=iPadPro6_129,IOS
 +DeviceProfileNameAndTypes=AppleTV,IOS
 +DeviceProfileNameAndTypes=AppleTV4K,IOS
++DeviceProfileNameAndTypes=AppleTV2_4K,IOS
 +DeviceProfileNameAndTypes=TVOS,AppleTV
 +DeviceProfileNameAndTypes=MacEditor,Mac
 +DeviceProfileNameAndTypes=MacClient,Mac
@@ -193,6 +203,11 @@ iPhone14,2=iPhone13Pro
 iPhone14,3=iPhone13ProMax
 iPhone14,4=iPhone13Mini
 iPhone14,5=iPhone13
+iPhone14,6=iPhoneSE3
+iPhone14,7=iPhone14
+iPhone14,8=iPhone14Plus
+iPhone15,2=iPhone14Pro
+iPhone15,3=iPhone14ProMax
 iPod9,1=iPodTouch7
 iPad5,[1,2]=iPadMini4
 iPad5,[3,4]=iPadAir2
@@ -218,14 +233,19 @@ iPad12,[1,2]=iPad9
 iPad13,[1,2]=iPadAir4
 iPad13,[4-7]=iPadPro3_11
 iPad13,[8-11]=iPadPro5_129
+iPad13,[16-17]=iPadAir5
+iPad13,[18-19]=iPad10
 iPad14,[1,2]=iPadMini6
+iPad14,[3,4]=iPadPro4_11
+iPad14,[5,6]=iPadPro6_129
 AppleTV5,=AppleTV
 AppleTV6,=AppleTV4K
+AppleTV11,=AppleTV2_4K
 ; Below are generic fallbacks for unrecognized devices. Add any specific devices above this section.
-iPhone=iPhone13
+iPhone=iPhone14
 iPod=iPodTouch7
-iPad=iPad9
-AppleTV=AppleTV4K
+iPad=iPad10
+AppleTV=AppleTV2_4K
 
 [IOS DeviceProfile]
 DeviceType=IOS
@@ -405,6 +425,29 @@ DeviceType=IOS
 BaseProfileName=iPhone11ProMax
 +CVars=ios.PhysicalScreenDensity=458
 
+[iPhoneSE3 DeviceProfile]
+DeviceType=IOS
+BaseProfileName=iPhone13
++CVars=ios.PhysicalScreenDensity=326
+
+[iPhone14 DeviceProfile]
+DeviceType=IOS
+BaseProfileName=iPhone13
+
+[iPhone14Plus DeviceProfile]
+DeviceType=IOS
+BaseProfileName=iPhone13
++CVars=ios.PhysicalScreenDensity=458
+
+[iPhone14Pro DeviceProfile]
+DeviceType=IOS
+BaseProfileName=iPhone13Pro
+
+[iPhone14ProMax DeviceProfile]
+DeviceType=IOS
+BaseProfileName=iPhone13ProMax
++CVars=ios.PhysicalScreenDensity=460
+
 [iPodTouch7 DeviceProfile]
 DeviceType=IOS
 BaseProfileName=iPhone7
@@ -473,6 +516,10 @@ DeviceType=IOS
 BaseProfileName=IPadPro
 +CVars=ios.PhysicalScreenDensity=264
 
+[iPadPro4_11 DeviceProfile]
+DeviceType=IOS
+BaseProfileName=IPadPro3_11
+
 [iPad5 DeviceProfile]
 DeviceType=IOS
 BaseProfileName=IPadPro97
@@ -498,6 +545,10 @@ DeviceType=IOS
 BaseProfileName=IPadPro97
 +CVars=ios.PhysicalScreenDensity=264
 
+[iPad10 DeviceProfile]
+DeviceType=IOS
+BaseProfileName=IPad9
+
 [iPadPro2_129 DeviceProfile]
 DeviceType=IOS
 BaseProfileName=IPadPro
@@ -518,6 +569,10 @@ DeviceType=IOS
 BaseProfileName=IPadPro
 +CVars=ios.PhysicalScreenDensity=264
 
+[iPadPro6_129 DeviceProfile]
+DeviceType=IOS
+BaseProfileName=IPadPro5_129
+
 [iPadAir3 DeviceProfile]
 DeviceType=IOS
 BaseProfileName=IPadPro
@@ -528,6 +583,10 @@ DeviceType=IOS
 BaseProfileName=IPadPro
 +CVars=ios.PhysicalScreenDensity=264
 
+[iPadAir5 DeviceProfile]
+DeviceType=IOS
+BaseProfileName=IPadAir4
+
 [iPadMini5 DeviceProfile]
 DeviceType=IOS
 BaseProfileName=IPadPro
@@ -552,6 +611,10 @@ DeviceType=IOS
 BaseProfileName=AppleTV
 +CVars=ios.PhysicalScreenDensity=0
 
+[AppleTV2_4K DeviceProfile]
+DeviceType=IOS
+BaseProfileName=AppleTV4K
+
 ;----------------
 ; Android devices
 
