Added Unreal Cloud DDC to the base cache graphs

It is disabled by default outside of Epic, and can be enabled by configuring the Default key in the StorageServers section of the Engine config.

diff --git a/Engine/Config/BaseEngine.ini b/Engine/Config/BaseEngine.ini
--- a/Engine/Config/BaseEngine.ini
+++ b/Engine/Config/BaseEngine.ini
@@ -2202,24 +2202,40 @@ ExtraArgs=--http asio --gc-cache-duration-seconds 2937600 --gc-interval-seconds
 HostName=localhost
 Port=1337
 
+; Each key defines settings for an Unreal Cloud DDC server.
+; The key is used in the ServerID field of a Cloud node in a graph for DDC.
+;   Host: The URI for the server to connect to, e.g., https://jupiter.example.com.
+;   EnvHostOverride: The name of an environment variable that can override the host, e.g., UE-CloudDataCacheHost.
+;   CommandLineHostOverride: The name of a command line switch that can override the host, e.g., CloudDataCacheHost.
+;   HostPinnedPublicKeys: A semicolon-delimited list of pinned public keys for the host.
+;   ResolveHostCanonicalName: Whether to resolve the host to a canonical name, e.g., for pinning to a region. Default is true.
+;   Namespace: The cache namespace to use on the server. Required.
+;   OAuthProvider: The URI for the OAuth provider to use for authentication. Optional.
+;   OAuthClientId: The client ID to use with the OAuth provider. Optional.
+;   OAuthSecret: The secret to use with the OAuth provider. Optional. May start with file:// to load the secret from a file.
+;   CommandLineOAuthSecretOverride: The name of a command line switch that can override the OAuth secret, e.g., CloudDataCacheOAuthSecret. Optional.
+;   OAuthScope: The scope to use with the OAuth provider. Default is cache_access.
+;   OAuthProviderIdentifier: The provider ID when using OIDC. Optional.
+;   OAuthAccessTokenEnvOverride: The name of a command line switch that can provide the OAuth access token to use, e.g., UE-CloudDataCacheAccessToken. Optional.
+;   OAuthPinnedPublicKeys: A semicolon-delimited list of pinned public keys for the OAuth provider. Optional.
+;   ReadOnly: Whether the cache is treated as read-only. Default is false.
+[StorageServers]
+Default=(Host=None, Namespace="ue.ddc")
+
 ; Derived Data backend graphs
 ; Each of the below backend graph sections contains a set of nodes used to create derived data backend graph.
 ; DerivedDataBackendGraph is the default graph for source builds, and InstalledDerivedDataBackendGraph is the default for installed builds.
 ; Others can be specified on the command line using:
 ;     -DDC=GraphSectionName     (for example: -DDC=NoShared)
 ; Each graph should start with a 'Root' node. The names of the other nodes are not predefined.
-; Supported node types are: KeyLength, AsyncPut, Hierarchical, FileSystem, Zen, Http, S3, ReadPak, WritePak, Verify
+; Supported node types are: KeyLength, AsyncPut, Hierarchical, FileSystem, Zen, Cloud, S3, ReadPak, WritePak, Verify
 ; The order that nodes are defined in is not relevant.
 ; FileSystem backends can be disabled by setting the corresponding environment variable to 'None'. E.g., UE-SharedDataCachePath=None
-; Http backends can be disabled by setting the corresponding environment variable to 'None'. E.g., UE-CloudDataCachePath=None
+; Cloud backends can be disabled by setting the corresponding environment variable to 'None'. E.g., UE-CloudDataCachePath=None
 
 [DerivedDataBackendGraph]
-; Configure the root entry. It uses a KeyLength node to hash long strings then passes requests to AsyncPut
-Root=(Type=KeyLength, Length=120, Inner=AsyncPut)
-; Configure the AsyncPut entry. It uses an AsyncPut node that then passes requests to Hierarchy
-AsyncPut=(Type=AsyncPut, Inner=Hierarchy)
-; Configure the hierarchy entry. This contains multiple nodes that are used in order until a read is found. Writes go to all writable backends.
-Hierarchy=(Type=Hierarchical, Inner=Pak, Inner=EnginePak, Inner=Local, Inner=Shared)
+; Configure the hierarchy as the root of the cache. Reads access the nodes in order until satisfied. Writes go to every writable node.
+Root=(Type=Hierarchical, Inner=Pak, Inner=EnginePak, Inner=Local, Inner=Shared, Inner=Cloud)
 ; Configure the local cache. This is a FileSystem cache with the following parameters:
 ;	ReadOnly: Can data be written to this layer. Default is false.
 ;	Touch: Update timestamps on every access instead of only after 1/4 of UnusedFileAge. Default is false.
@@ -2241,8 +2257,10 @@ Hierarchy=(Type=Hierarchical, Inner=Pak, Inner=EnginePak, Inner=Local, Inner=Sha
 ;		LegacyOnly: Look only for uncompressed data in the pre-5.0 directory layout.
 ;		ValueWithLegacyFallback: Look in the 5.0+ location but fall back to the pre-5.0 location on miss.
 Local=(Type=FileSystem, PurgeTransient=true, UnusedFileAge=34, PromptIfMissing=true, Path=%ENGINEDIR%DerivedDataCache, EnvPathOverride=UE-LocalDataCachePath, EditorOverrideSetting=LocalDerivedDataCache, CommandLineOverride=LocalDataCachePath)
-; Configure the shared DDC that is accessed after local. It's a filesystem DDC and the parameters are explained above
+; Configure the shared cache that is accessed after local.
 Shared=(Type=FileSystem, UnusedFileAge=10, FoldersToClean=10, ConsiderSlowAt=70, Path=?EpicDDC, EnvPathOverride=UE-SharedDataCachePath, EditorOverrideSetting=SharedDerivedDataCache, CommandLineOverride=SharedDataCachePath)
+; Configure the cloud cache that is accessed after shared.
+Cloud=(Type=Cloud, ServerID=Default, EnvHostOverride=UE-CloudDataCacheHost, CommandLineHostOverride=CloudDataCacheHost, CommandLineOAuthSecretOverride="CloudDataCacheOAuthSecret", OAuthAccessTokenEnvOverride="UE-CloudDataCacheAccessToken")
 ; Configure a Project Pak node. This is a pre-generated DDC data file for the project that can be distributed to reduce runtime fetches/generation
 ; See documentation for how to create a DDP via the DerivedDataCache commandlet
 Pak=(Type=ReadPak, Filename="%GAMEDIR%DerivedDataCache/DDC.ddp")
@@ -2252,9 +2270,10 @@ EnginePak=(Type=ReadPak, Filename="%ENGINEDIR%DerivedDataCache/DDC.ddp")
 [InstalledDerivedDataBackendGraph]
 Root=(Type=KeyLength, Length=120, Inner=AsyncPut)
 AsyncPut=(Type=AsyncPut, Inner=Hierarchy)
-Hierarchy=(Type=Hierarchical, Inner=Pak, Inner=CompressedPak, Inner=EnginePak, Inner=EnterprisePak, Inner=Local, Inner=Shared)
+Hierarchy=(Type=Hierarchical, Inner=Pak, Inner=CompressedPak, Inner=EnginePak, Inner=EnterprisePak, Inner=Local, Inner=Shared, Inner=Cloud)
 Local=(Type=FileSystem, PurgeTransient=true, UnusedFileAge=34, Path="%ENGINEVERSIONAGNOSTICUSERDIR%DerivedDataCache", EnvPathOverride=UE-LocalDataCachePath, EditorOverrideSetting=LocalDerivedDataCache)
 Shared=(Type=FileSystem, UnusedFileAge=10, FoldersToClean=10, ConsiderSlowAt=70, Path=?EpicDDC, EnvPathOverride=UE-SharedDataCachePath, EditorOverrideSetting=SharedDerivedDataCache, CommandLineOverride=SharedDataCachePath)
+Cloud=(Type=Cloud, ServerID=Default, EnvHostOverride=UE-CloudDataCacheHost, CommandLineHostOverride=CloudDataCacheHost, CommandLineOAuthSecretOverride="CloudDataCacheOAuthSecret", OAuthAccessTokenEnvOverride="UE-CloudDataCacheAccessToken")
 Pak=(Type=ReadPak, Filename="%GAMEDIR%DerivedDataCache/DDC.ddp")
 CompressedPak=(Type=ReadPak, Filename="%GAMEDIR%DerivedDataCache/Compressed.ddp", Compressed=true)
 EnginePak=(Type=ReadPak, Filename=../../../Engine/DerivedDataCache/Compressed.ddp, Compressed=true)
@@ -2307,18 +2326,20 @@ Local=(Type=FileSystem, Clean=true, PurgeTransient=true, UnusedFileAge=1, Path="
 [CreatePak]
 Root=(Type=KeyLength, Length=120, Inner=AsyncPut)
 AsyncPut=(Type=AsyncPut, Inner=Hierarchy)
-Hierarchy=(Type=Hierarchical, Inner=PakWrite, Inner=PakRead, Inner=Local, Inner=Shared)
+Hierarchy=(Type=Hierarchical, Inner=PakWrite, Inner=PakRead, Inner=Local, Inner=Shared, Inner=Cloud)
 Local=(Type=FileSystem, PurgeTransient=true, UnusedFileAge=34, Path=%ENGINEDIR%DerivedDataCache, EnvPathOverride=UE-LocalDataCachePath, EditorOverrideSetting=LocalDerivedDataCache)
 Shared=(Type=FileSystem, DeleteUnused=false, ConsiderSlowAt=70, Path=?EpicDDC, EnvPathOverride=UE-SharedDataCachePath, EditorOverrideSetting=SharedDerivedDataCache, CommandLineOverride=SharedDataCachePath)
+Cloud=(Type=Cloud, ServerID=Default, EnvHostOverride=UE-CloudDataCacheHost, CommandLineHostOverride=CloudDataCacheHost, CommandLineOAuthSecretOverride="CloudDataCacheOAuthSecret", OAuthAccessTokenEnvOverride="UE-CloudDataCacheAccessToken")
 PakRead=(Type=ReadPak, Filename="%GAMEDIR%DerivedDataCache/DDC.ddp")
 PakWrite=(Type=WritePak, Filename="%GAMEDIR%DerivedDataCache/DDC.ddp")
 
 [CreateInstalledProjectPak]
 Root=(Type=KeyLength, Length=120, Inner=AsyncPut)
 AsyncPut=(Type=AsyncPut, Inner=Hierarchy)
-Hierarchy=(Type=Hierarchical, Inner=EnginePak, Inner=PakWrite, Inner=PakRead, Inner=Local, Inner=Shared)
+Hierarchy=(Type=Hierarchical, Inner=EnginePak, Inner=PakWrite, Inner=PakRead, Inner=Local, Inner=Shared, Inner=Cloud)
 Local=(Type=FileSystem, PurgeTransient=true, UnusedFileAge=34, Path=%ENGINEDIR%DerivedDataCache, EnvPathOverride=UE-LocalDataCachePath, EditorOverrideSetting=LocalDerivedDataCache)
 Shared=(Type=FileSystem, DeleteUnused=false, ConsiderSlowAt=70, Path=?EpicDDC, EnvPathOverride=UE-SharedDataCachePath, EditorOverrideSetting=SharedDerivedDataCache, CommandLineOverride=SharedDataCachePath)
+Cloud=(Type=Cloud, ServerID=Default, EnvHostOverride=UE-CloudDataCacheHost, CommandLineHostOverride=CloudDataCacheHost, CommandLineOAuthSecretOverride="CloudDataCacheOAuthSecret", OAuthAccessTokenEnvOverride="UE-CloudDataCacheAccessToken")
 EnginePak=(Type=ReadPak, Filename=../../../Engine/DerivedDataCache/Compressed.ddp, Compressed=true)
 PakRead=(Type=ReadPak, Filename="%GAMEDIR%DerivedDataCache/Compressed.ddp", Compressed=true)
 PakWrite=(Type=WritePak, Filename="%GAMEDIR%DerivedDataCache/Compressed.ddp", Compressed=true)
@@ -2326,26 +2347,29 @@ PakWrite=(Type=WritePak, Filename="%GAMEDIR%DerivedDataCache/Compressed.ddp", Co
 [CreateInstalledEnginePak]
 Root=(Type=KeyLength, Length=120, Inner=AsyncPut)
 AsyncPut=(Type=AsyncPut, Inner=Hierarchy)
-Hierarchy=(Type=Hierarchical, Inner=PakWrite, Inner=Local, Inner=Shared)
+Hierarchy=(Type=Hierarchical, Inner=PakWrite, Inner=Local, Inner=Shared, Inner=Cloud)
 Local=(Type=FileSystem, PurgeTransient=true, UnusedFileAge=34, Path=%ENGINEDIR%DerivedDataCache, EnvPathOverride=UE-LocalDataCachePath, EditorOverrideSetting=LocalDerivedDataCache)
 Shared=(Type=FileSystem, DeleteUnused=false, ConsiderSlowAt=70, Path=?EpicDDC, EnvPathOverride=UE-SharedDataCachePath, EditorOverrideSetting=SharedDerivedDataCache, CommandLineOverride=SharedDataCachePath)
+Cloud=(Type=Cloud, ServerID=Default, EnvHostOverride=UE-CloudDataCacheHost, CommandLineHostOverride=CloudDataCacheHost, CommandLineOAuthSecretOverride="CloudDataCacheOAuthSecret", OAuthAccessTokenEnvOverride="UE-CloudDataCacheAccessToken")
 PakWrite=(Type=WritePak, Filename="%ENGINEDIR%DerivedDataCache/Compressed.ddp", Compressed=true)
 
 [CreateInstalledEnterprisePak]
 Root=(Type=KeyLength, Length=120, Inner=AsyncPut)
 AsyncPut=(Type=AsyncPut, Inner=Hierarchy)
-Hierarchy=(Type=Hierarchical, Inner=PakWrite, Inner=Local, Inner=Shared)
+Hierarchy=(Type=Hierarchical, Inner=PakWrite, Inner=Local, Inner=Shared, Inner=Cloud)
 Local=(Type=FileSystem, PurgeTransient=true, UnusedFileAge=34, Path=%ENGINEDIR%DerivedDataCache, EnvPathOverride=UE-LocalDataCachePath, EditorOverrideSetting=LocalDerivedDataCache)
 Shared=(Type=FileSystem, DeleteUnused=false, ConsiderSlowAt=70, Path=?EpicDDC, EnvPathOverride=UE-SharedDataCachePath, EditorOverrideSetting=SharedDerivedDataCache, CommandLineOverride=SharedDataCachePath)
+Cloud=(Type=Cloud, ServerID=Default, EnvHostOverride=UE-CloudDataCacheHost, CommandLineHostOverride=CloudDataCacheHost, CommandLineOAuthSecretOverride="CloudDataCacheOAuthSecret", OAuthAccessTokenEnvOverride="UE-CloudDataCacheAccessToken")
 PakWrite=(Type=WritePak, Filename=../../../Enterprise/DerivedDataCache/Compressed.ddp, Compressed=true)
 
 [CreateProjectCache]
 Root=(Type=KeyLength, Length=120, Inner=AsyncPut)
 AsyncPut=(Type=AsyncPut, Inner=Hierarchy)
-Hierarchy=(Type=Hierarchical, Inner=Local, Inner=Project, Inner=Shared)
+Hierarchy=(Type=Hierarchical, Inner=Local, Inner=Project, Inner=Shared, Inner=Cloud)
 Local=(Type=FileSystem, PurgeTransient=true, UnusedFileAge=34, Path=%ENGINEDIR%DerivedDataCache, EnvPathOverride=UE-LocalDataCachePath, EditorOverrideSetting=LocalDerivedDataCache)
 Project=(Type=FileSystem, Flush=true, PurgeTransient=true, UnusedFileAge=34, Path=%GAMEDIR%ProjectDerivedData)
 Shared=(Type=FileSystem, DeleteUnused=false, ConsiderSlowAt=70, Path=?EpicDDC, EnvPathOverride=UE-SharedDataCachePath, EditorOverrideSetting=SharedDerivedDataCache, CommandLineOverride=SharedDataCachePath)
+Cloud=(Type=Cloud, ServerID=Default, EnvHostOverride=UE-CloudDataCacheHost, CommandLineHostOverride=CloudDataCacheHost, CommandLineOAuthSecretOverride="CloudDataCacheOAuthSecret", OAuthAccessTokenEnvOverride="UE-CloudDataCacheAccessToken")
 
 [VirtualTextureChunkDDCCache]
 UnusedFileAge=34
diff --git a/Engine/Source/Developer/DerivedDataCache/Private/DerivedDataBackends.cpp b/Engine/Source/Developer/DerivedDataCache/Private/DerivedDataBackends.cpp
--- a/Engine/Source/Developer/DerivedDataCache/Private/DerivedDataBackends.cpp
+++ b/Engine/Source/Developer/DerivedDataCache/Private/DerivedDataBackends.cpp
@@ -347,7 +347,7 @@ public:
 				{
 					ParsedNode = MakeTuple(ParseS3Cache(*NodeName, *Entry), ECacheStoreFlags::Local | ECacheStoreFlags::Query | ECacheStoreFlags::StopStore);
 				}
-				else if (NodeType == TEXT("Http"))
+				else if (NodeType == TEXT("Cloud") || NodeType == TEXT("Http"))
 				{
 					ParsedNode = CreateHttpCacheStore(*NodeName, *Entry);
 				}
diff --git a/Engine/Source/Developer/DerivedDataCache/Private/HttpCacheStore.cpp b/Engine/Source/Developer/DerivedDataCache/Private/HttpCacheStore.cpp
--- a/Engine/Source/Developer/DerivedDataCache/Private/HttpCacheStore.cpp
+++ b/Engine/Source/Developer/DerivedDataCache/Private/HttpCacheStore.cpp
@@ -2752,7 +2752,7 @@ TTuple<ILegacyCacheStore*, ECacheStoreFlags> CreateHttpCacheStore(const TCHAR* N
 		UE_LOG(LogDerivedDataCache, Warning, TEXT("%s: Missing required parameter 'StructuredNamespace', falling back to '%s'"), NodeName, *Params.Namespace);
 	}
 
-	if (!Params.Host.StartsWith(TEXT("http://localhost")))
+	if (bValidParams && !Params.Host.StartsWith(TEXT("http://localhost")))
 	{
 		bool bValidOAuthAccessToken = !Params.OAuthAccessToken.IsEmpty();
 
