Moved Distance Field AO out of ShadowQuality and into GlobalIllumination Quality scalability group, so that Medium GI gives DFAO
Moved Screen Space Reflection out of EffectsQuality and into ReflectionsQuality scalability group, so that Medium Reflections gives SSR

diff --git a/Engine/Config/BaseDeviceProfiles.ini b/Engine/Config/BaseDeviceProfiles.ini
--- a/Engine/Config/BaseDeviceProfiles.ini
+++ b/Engine/Config/BaseDeviceProfiles.ini
@@ -268,6 +268,8 @@ bIsVisibleForAssets=True
 ; Scalability
 +CVars=sg.AntiAliasingQuality=0
 +CVars=sg.ShadowQuality=1
++CVars=sg.GlobalIlluminationQuality=1
++CVars=sg.ReflectionQuality=1
 +CVars=sg.ViewDistanceQuality=3
 +CVars=sg.PostProcessQuality=3
 +CVars=sg.TextureQuality=3
@@ -284,6 +286,7 @@ BaseProfileName=IOS
 +CVars=r.MobileContentScaleFactor=2
 +CVars=sg.PostProcessQuality=3
 +CVars=sg.ShadowQuality=2
++CVars=sg.GlobalIlluminationQuality=2
 +CVars=r.Mobile.AmbientOcclusionQuality=0
 +CVars=r.Mobile.PixelProjectedReflectionQuality=0
 
@@ -294,6 +297,7 @@ BaseProfileName=IOS
 +CVars=r.MobileContentScaleFactor=2
 +CVars=sg.PostProcessQuality=3
 +CVars=sg.ShadowQuality=2
++CVars=sg.GlobalIlluminationQuality=2
 +CVars=r.Mobile.AmbientOcclusionQuality=0
 +CVars=r.Mobile.PixelProjectedReflectionQuality=0
 
@@ -327,6 +331,7 @@ DeviceType=IOS
 BaseProfileName=iPhone7
 +CVars=ios.PhysicalScreenDensity=326
 +CVars=sg.ShadowQuality=3
++CVars=sg.GlobalIlluminationQuality=3
 +CVars=r.Mobile.AmbientOcclusionQuality=1
 +CVars=r.Mobile.PixelProjectedReflectionQuality=1
 
@@ -335,6 +340,7 @@ DeviceType=IOS
 BaseProfileName=iPhone7Plus
 +CVars=ios.PhysicalScreenDensity=401
 +CVars=sg.ShadowQuality=3
++CVars=sg.GlobalIlluminationQuality=3
 +CVars=r.Mobile.AmbientOcclusionQuality=1
 +CVars=r.Mobile.PixelProjectedReflectionQuality=1
 
@@ -462,6 +468,7 @@ BaseProfileName=IOS
 +CVars=ios.PhysicalScreenDensity=264
 +CVars=sg.PostProcessQuality=3
 +CVars=sg.ShadowQuality=2
++CVars=sg.GlobalIlluminationQuality=2
 +CVars=r.Mobile.AmbientOcclusionQuality=0
 +CVars=r.Mobile.PixelProjectedReflectionQuality=0
 
@@ -483,6 +490,7 @@ BaseProfileName=IOS
 +CVars=r.MetalComputeParameterSize=1024
 +CVars=r.EarlyZPass=3
 +CVars=sg.ShadowQuality=3
++CVars=sg.GlobalIlluminationQuality=3
 +CVars=sg.PostProcessQuality=3
 +CVars=sg.AntiAliasingQuality=3
 
@@ -603,6 +611,7 @@ BaseProfileName=IOS
 +CVars=ios.PhysicalScreenDensity=0
 +CVars=sg.PostProcessQuality=3
 +CVars=sg.ShadowQuality=3
++CVars=sg.GlobalIlluminationQuality=3
 +CVars=r.Mobile.AmbientOcclusionQuality=0
 +CVars=r.Mobile.PixelProjectedReflectionQuality=0
 
@@ -710,6 +719,8 @@ BaseProfileName=Android
 +CVars=sg.ViewDistanceQuality=0
 +CVars=sg.AntiAliasingQuality=0
 +CVars=sg.ShadowQuality=0
++CVars=sg.GlobalIlluminationQuality=0
++CVars=sg.ReflectionQuality=0
 +CVars=sg.PostProcessQuality=0
 +CVars=sg.TextureQuality=0
 +CVars=sg.EffectsQuality=0
@@ -722,6 +733,8 @@ BaseProfileName=Android
 +CVars=sg.ViewDistanceQuality=1
 +CVars=sg.AntiAliasingQuality=1
 +CVars=sg.ShadowQuality=1
++CVars=sg.GlobalIlluminationQuality=1
++CVars=sg.ReflectionQuality=1
 +CVars=sg.PostProcessQuality=1
 +CVars=sg.TextureQuality=1
 +CVars=sg.EffectsQuality=1
@@ -733,6 +746,8 @@ BaseProfileName=Android
 +CVars=sg.ViewDistanceQuality=2
 +CVars=sg.AntiAliasingQuality=2
 +CVars=sg.ShadowQuality=2
++CVars=sg.GlobalIlluminationQuality=2
++CVars=sg.ReflectionQuality=2
 +CVars=sg.PostProcessQuality=2
 +CVars=sg.TextureQuality=2
 +CVars=sg.EffectsQuality=2
@@ -746,6 +761,8 @@ BaseProfileName=Android
 +CVars=sg.ViewDistanceQuality=2
 +CVars=sg.AntiAliasingQuality=1
 +CVars=sg.ShadowQuality=2
++CVars=sg.GlobalIlluminationQuality=2
++CVars=sg.ReflectionQuality=2
 +CVars=sg.PostProcessQuality=2
 +CVars=sg.TextureQuality=2
 +CVars=sg.EffectsQuality=2
diff --git a/Engine/Config/BaseScalability.ini b/Engine/Config/BaseScalability.ini
--- a/Engine/Config/BaseScalability.ini
+++ b/Engine/Config/BaseScalability.ini
@@ -127,7 +127,6 @@ r.Shadow.DistanceScale=0.6
 r.Shadow.CSM.TransitionScale=0
 r.Shadow.PreShadowResolutionFactor=0.5
 r.DistanceFieldShadowing=0
-r.DistanceFieldAO=0
 r.VolumetricFog=0
 r.LightMaxDrawDistanceScale=0
 r.CapsuleShadows=0
@@ -150,8 +149,6 @@ r.Shadow.DistanceScale=0.7
 r.Shadow.CSM.TransitionScale=0.25
 r.Shadow.PreShadowResolutionFactor=0.5
 r.DistanceFieldShadowing=0
-; Enable DFAO when Lumen is off
-r.DistanceFieldAO=1
 r.VolumetricFog=0
 r.LightMaxDrawDistanceScale=.5
 r.CapsuleShadows=1
@@ -174,8 +171,6 @@ r.Shadow.DistanceScale=0.85
 r.Shadow.CSM.TransitionScale=0.8
 r.Shadow.PreShadowResolutionFactor=0.5
 r.DistanceFieldShadowing=1
-r.DistanceFieldAO=1
-r.AOQuality=1
 r.VolumetricFog=1
 r.VolumetricFog.GridPixelSize=16
 r.VolumetricFog.GridSizeZ=64
@@ -201,8 +196,6 @@ r.Shadow.DistanceScale=1.0
 r.Shadow.CSM.TransitionScale=1.0
 r.Shadow.PreShadowResolutionFactor=1.0
 r.DistanceFieldShadowing=1
-r.DistanceFieldAO=1
-r.AOQuality=2
 r.VolumetricFog=1
 r.VolumetricFog.GridPixelSize=8
 r.VolumetricFog.GridSizeZ=128
@@ -228,8 +221,6 @@ r.Shadow.DistanceScale=1.0
 r.Shadow.CSM.TransitionScale=1.0
 r.Shadow.PreShadowResolutionFactor=1.0
 r.DistanceFieldShadowing=1
-r.DistanceFieldAO=1
-r.AOQuality=2
 r.VolumetricFog=1
 r.VolumetricFog.GridPixelSize=4
 r.VolumetricFog.GridSizeZ=128
@@ -247,12 +238,18 @@ r.Shadow.Virtual.SMRT.SamplesPerRayLocal=8
 ;-----------------------------------------------------------------------------------------------------------------
 
 [GlobalIlluminationQuality@0]
+r.DistanceFieldAO=0
 r.Lumen.DiffuseIndirect.Allow=0
 
 [GlobalIlluminationQuality@1]
+; Enable DFAO when Lumen is off
+r.DistanceFieldAO=1
+r.AOQuality=1
 r.Lumen.DiffuseIndirect.Allow=0
 
 [GlobalIlluminationQuality@2]
+r.DistanceFieldAO=1
+r.AOQuality=1
 r.Lumen.DiffuseIndirect.Allow=1
 r.LumenScene.Radiosity.ProbeSpacing=8
 r.LumenScene.Radiosity.HemisphereProbeResolution=3
@@ -274,6 +271,8 @@ r.Lumen.TranslucencyVolume.RadianceCache.NumProbesToTraceBudget=200
 r.LumenScene.SurfaceCache.CardCaptureRefreshFraction=0
 
 [GlobalIlluminationQuality@3]
+r.DistanceFieldAO=1
+r.AOQuality=2
 r.Lumen.DiffuseIndirect.Allow=1
 r.LumenScene.Radiosity.ProbeSpacing=4
 r.LumenScene.Radiosity.HemisphereProbeResolution=4
@@ -296,6 +295,8 @@ r.Lumen.TranslucencyVolume.RadianceCache.NumProbesToTraceBudget=200
 r.LumenScene.SurfaceCache.CardCaptureRefreshFraction=0.125
 
 [GlobalIlluminationQuality@Cine]
+r.DistanceFieldAO=1
+r.AOQuality=2
 r.Lumen.DiffuseIndirect.Allow=1
 r.LumenScene.Radiosity.ProbeSpacing=4
 r.LumenScene.Radiosity.HemisphereProbeResolution=4
@@ -321,24 +322,34 @@ r.LumenScene.SurfaceCache.CardCaptureRefreshFraction=0.125
 ;-----------------------------------------------------------------------------------------------------------------
 
 [ReflectionQuality@0]
+r.SSR.Quality=0
+r.SSR.HalfResSceneColor=1
 r.Lumen.Reflections.Allow=0
 
 [ReflectionQuality@1]
+r.SSR.Quality=2
+r.SSR.HalfResSceneColor=1
 r.Lumen.Reflections.Allow=0
 
 [ReflectionQuality@2]
+r.SSR.Quality=2
+r.SSR.HalfResSceneColor=1
 r.Lumen.Reflections.Allow=1
 r.Lumen.Reflections.DownsampleFactor=2
 r.Lumen.TranslucencyReflections.FrontLayer.Allow=0
 r.Lumen.TranslucencyReflections.FrontLayer.Enable=0
 
 [ReflectionQuality@3]
+r.SSR.Quality=3
+r.SSR.HalfResSceneColor=0
 r.Lumen.Reflections.Allow=1
 r.Lumen.Reflections.DownsampleFactor=1
 r.Lumen.TranslucencyReflections.FrontLayer.Allow=1
 r.Lumen.TranslucencyReflections.FrontLayer.Enable=0
 
 [ReflectionQuality@Cine]
+r.SSR.Quality=4
+r.SSR.HalfResSceneColor=0
 r.Lumen.Reflections.Allow=1
 r.Lumen.Reflections.DownsampleFactor=1
 r.Lumen.TranslucencyReflections.FrontLayer.Allow=1
@@ -576,8 +587,6 @@ r.Streaming.MaxEffectiveScreenSize=0
 [EffectsQuality@0]
 r.TranslucencyLightingVolumeDim=24
 r.RefractionQuality=0
-r.SSR.Quality=0
-r.SSR.HalfResSceneColor=1
 r.SceneColorFormat=3
 r.DetailMode=0
 r.TranslucencyVolumeBlur=0
@@ -606,8 +615,6 @@ fx.Niagara.QualityLevel=0
 [EffectsQuality@1]
 r.TranslucencyLightingVolumeDim=32
 r.RefractionQuality=0
-r.SSR.Quality=0
-r.SSR.HalfResSceneColor=1
 r.SceneColorFormat=3
 r.DetailMode=1
 r.TranslucencyVolumeBlur=0
@@ -636,8 +643,6 @@ fx.Niagara.QualityLevel=1
 [EffectsQuality@2]
 r.TranslucencyLightingVolumeDim=48
 r.RefractionQuality=2
-r.SSR.Quality=2
-r.SSR.HalfResSceneColor=1
 r.SceneColorFormat=3
 r.DetailMode=1
 r.TranslucencyVolumeBlur=1
@@ -666,8 +671,6 @@ fx.Niagara.QualityLevel=2
 [EffectsQuality@3]
 r.TranslucencyLightingVolumeDim=64
 r.RefractionQuality=2
-r.SSR.Quality=3
-r.SSR.HalfResSceneColor=0
 r.SceneColorFormat=4
 r.DetailMode=2
 r.TranslucencyVolumeBlur=1
@@ -696,8 +699,6 @@ fx.Niagara.QualityLevel=3
 [EffectsQuality@Cine]
 r.TranslucencyLightingVolumeDim=64
 r.RefractionQuality=2
-r.SSR.Quality=4
-r.SSR.HalfResSceneColor=0
 r.SceneColorFormat=4
 r.DetailMode=2
 r.TranslucencyVolumeBlur=1
