USD: Add LevelSequenceEditorSettings config to automatically bind stage actor Time tracks when binding them to the Sequencer.

diff --git a/Engine/Config/BaseEditorPerProjectUserSettings.ini b/Engine/Config/BaseEditorPerProjectUserSettings.ini
--- a/Engine/Config/BaseEditorPerProjectUserSettings.ini
+++ b/Engine/Config/BaseEditorPerProjectUserSettings.ini
@@ -898,6 +898,7 @@ FrameNumberDisplayFormat=Frames
 +TrackSettings=(MatchingActorClass=/Script/Engine.CameraActor,DefaultTracks=(/Script/MovieSceneTracks.MovieScene3DTransformTrack),DefaultPropertyTracks=((ComponentPath="CameraComponent",PropertyPath="FieldOfView")))
 +TrackSettings=(MatchingActorClass=/Script/CinematicCamera.CineCameraActor,DefaultPropertyTracks=((ComponentPath="CameraComponent",PropertyPath="CurrentFocalLength"),(ComponentPath="CameraComponent",PropertyPath="FocusSettings.ManualFocusDistance"),(ComponentPath="CameraComponent",PropertyPath="CurrentAperture")),ExcludeDefaultPropertyTracks=((ComponentPath="CameraComponent",PropertyPath="FieldOfView")))
 +TrackSettings=(MatchingActorClass=/Script/Engine.Light,DefaultTracks=(None),DefaultPropertyTracks=((ComponentPath="LightComponent0",PropertyPath="Intensity"),(ComponentPath="LightComponent0",PropertyPath="LightColor")))
++TrackSettings=(MatchingActorClass=/Script/UsdStage.UsdStageActor,DefaultTracks=(None),DefaultPropertyTracks=((ComponentPath="",PropertyPath="Time")))
 
 [/Script/LevelSequenceEditor.TakeRecorderEditorSettings]
 +TrackSettings=(MatchingActorClass=/Script/Engine.StaticMeshActor,DefaultTracks=(/Script/MovieSceneTracks.MovieScene3DTransformTrack))
