Enables TSR's anti-flickering on TSR high scalability settings to close the performance gap between medium and epic

diff --git a/Engine/Config/BaseScalability.ini b/Engine/Config/BaseScalability.ini
--- a/Engine/Config/BaseScalability.ini
+++ b/Engine/Config/BaseScalability.ini
@@ -64,7 +64,7 @@ r.TemporalAA.Quality=1
 r.TSR.History.R11G11B10=1
 r.TSR.History.ScreenPercentage=100
 r.TSR.History.UpdateQuality=2
-r.TSR.ShadingRejection.Flickering=0
+r.TSR.ShadingRejection.Flickering=1
 r.TSR.Velocity.Extrapolation=0
 r.TSR.Velocity.HoleFill=1
 r.TSR.RejectionAntiAliasingQuality=1
