Disabled r.Lumen.ScreenProbeGather.FullResolutionJitterWidth on GI High settings to reduce cost (.13ms at 60% SP on PS5), as r.Lumen.ScreenProbeGather.StochasticInterpolation is already injecting noise which helps with temporal stability.

diff --git a/Engine/Config/BaseScalability.ini b/Engine/Config/BaseScalability.ini
--- a/Engine/Config/BaseScalability.ini
+++ b/Engine/Config/BaseScalability.ini
@@ -261,7 +261,8 @@ r.Lumen.ScreenProbeGather.DownsampleFactor=32
 r.Lumen.ScreenProbeGather.TracingOctahedronResolution=8
 r.Lumen.ScreenProbeGather.IrradianceFormat=1
 r.Lumen.ScreenProbeGather.StochasticInterpolation=1
-r.Lumen.ScreenProbeGather.FullResolutionJitterWidth=1
+; FullResolutionJitterWidth is useful to distribute screen probe error spatially to improve temporal stability, but StochasticInterpolation already injects noise
+r.Lumen.ScreenProbeGather.FullResolutionJitterWidth=0
 r.Lumen.ScreenProbeGather.TwoSidedFoliageBackfaceDiffuse=0
 r.Lumen.Reflections.Foliage=0
 r.Lumen.TranslucencyVolume.GridPixelSize=64
