Reduce the iteratively cooked types to just sound and texture types until we've got more solid non-asset dependency tracking in place.

diff --git a/Engine/Config/BaseEditor.ini b/Engine/Config/BaseEditor.ini
--- a/Engine/Config/BaseEditor.ini
+++ b/Engine/Config/BaseEditor.ini
@@ -438,24 +438,24 @@ IterativeClassAllowList=/Script/Engine.Texture2D
 +IterativeClassAllowList=/Script/Engine.SoundNodeAttenuation
 
 ; StaticMesh and commonly co-packaged classes
-+IterativeClassAllowList=/Script/Engine.StaticMesh
-+IterativeClassAllowList=/Script/Engine.BodySetup
-+IterativeClassAllowList=/Script/NavigationSystem.NavCollision
-+IterativeClassAllowList=/Script/Engine.StaticMeshDescriptionBulkData
-+IterativeClassAllowList=/Script/UnrealEd.FbxStaticMeshImportData
-+IterativeClassAllowList=/Script/Engine.StaticMeshSocket
+; +IterativeClassAllowList=/Script/Engine.StaticMesh
+; +IterativeClassAllowList=/Script/Engine.BodySetup
+; +IterativeClassAllowList=/Script/NavigationSystem.NavCollision
+; +IterativeClassAllowList=/Script/Engine.StaticMeshDescriptionBulkData
+; +IterativeClassAllowList=/Script/UnrealEd.FbxStaticMeshImportData
+; +IterativeClassAllowList=/Script/Engine.StaticMeshSocket
 
 ; MaterialInstanceConstant and commonly co-packaged classes
-+IterativeClassAllowList=/Script/Engine.MaterialInstanceConstant
-+IterativeClassAllowList=/Script/Engine.SceneThumbnailInfoWithPrimitive
+; +IterativeClassAllowList=/Script/Engine.MaterialInstanceConstant
+; +IterativeClassAllowList=/Script/Engine.SceneThumbnailInfoWithPrimitive
 
 ; AnimSequence and commonly co-packaged classes
-+IterativeClassAllowList=/Script/Engine.AnimSequence
-+IterativeClassAllowList=/Script/Engine.AnimDataModel
-+IterativeClassAllowList=/Script/Engine.AnimCompress_PerTrackCompression
-+IterativeClassAllowList=/Script/Engine.AnimCompress_BitwiseCompressOnly
-+IterativeClassAllowList=/Script/AnimationModifiers.AnimationModifiersAssetUserData
-+IterativeClassAllowList=/Script/UnrealEd.FbxAnimSequenceImportData
+; +IterativeClassAllowList=/Script/Engine.AnimSequence
+; +IterativeClassAllowList=/Script/Engine.AnimDataModel
+; +IterativeClassAllowList=/Script/Engine.AnimCompress_PerTrackCompression
+; +IterativeClassAllowList=/Script/Engine.AnimCompress_BitwiseCompressOnly
+; +IterativeClassAllowList=/Script/AnimationModifiers.AnimationModifiersAssetUserData
+; +IterativeClassAllowList=/Script/UnrealEd.FbxAnimSequenceImportData
 
 [CookPlatformDataCacheSettings]
 MaterialInstance=100
@@ -602,4 +602,4 @@ bCanModifyPluginsFromBrowser=true
 
 [/Script/GameFeaturesEditor.GameFeaturesEditorSettings]
 +PluginTemplates=(Path=(Path="../../Plugins/Experimental/GameFeatures/Templates/GameFeaturePluginContentOnly"),Label=NSLOCTEXT("[/Script/GameFeaturesEditor]", "E17DB7E94A2C0F5968A52E89959864C0", "Game Feature (Content Only)"),Description=NSLOCTEXT("[/Script/GameFeaturesEditor]", "CB5136B64AFAFBA5EDEA078FCAC7BCC9", "Create a new Game Feature Plugin."),DefaultGameFeatureDataClass=None,DefaultGameFeatureDataName="")
-+PluginTemplates=(Path=(Path="../../Plugins/Experimental/GameFeatures/Templates/GameFeaturePluginWithCode"),Label=NSLOCTEXT("[/Script/GameFeaturesEditor]", "90D6FD2A455176CEC42DD79A23AEC3B0", "Game Feature (with C++)"),Description=NSLOCTEXT("[/Script/GameFeaturesEditor]", "2B00A0E7415C8CEAA6DA609A141237F0", "Create a new Game Feature Plugin with a minimal amount of code."),DefaultGameFeatureDataClass=None,DefaultGameFeatureDataName="")
++PluginTemplates=(Path=(Path="../../Plugins/Experimental/GameFeatures/Templates/GameFeaturePluginWithCode"),Label=NSLOCTEXT("[/Script/GameFeaturesEditor]", "90D6FD2A455176CEC42DD79A23AEC3B0", "Game Feature (with C++)"),Description=NSLOCTEXT("[/Script/GameFeaturesEditor]", "2B00A0E7415C8CEAA6DA609A141237F0", "Create a new Game Feature Plugin with a minimal amount of code."),DefaultGameFeatureDataClass=None,DefaultGameFeatureDataName="")
\ No newline at end of file
