Update ShooterGame, QAGame, and Lyra to use Zen for LOCAL DDC only.

diff --git a/Engine/Config/BaseEngine.ini b/Engine/Config/BaseEngine.ini
--- a/Engine/Config/BaseEngine.ini
+++ b/Engine/Config/BaseEngine.ini
@@ -2280,14 +2280,14 @@ EnterprisePak=(Type=ReadPak, Filename=../../../Enterprise/DerivedDataCache/Compr
 [ZenOnlyDDC]
 Root=(Type=KeyLength, Length=120, Inner=AsyncPut)
 AsyncPut=(Type=AsyncPut, Inner=Zen)
-Zen=(Type=Zen, Namespace="ue4.ddc", StructuredNamespace="ue.ddc")
+Zen=(Type=Zen, Namespace="ue.ddc")
 
 [ZenDDC]
-Root=(Type=KeyLength, Length=120, Inner=AsyncPut)
-AsyncPut=(Type=AsyncPut, Inner=Hierarchy)
-Hierarchy=(Type=Hierarchical, Inner=Pak, Inner=EnginePak, Inner=Zen, Inner=Shared)
-Zen=(Type=Zen, Namespace="ue4.ddc", StructuredNamespace="ue.ddc")
+Root=(Type=Hierarchical, Inner=Pak, Inner=EnginePak, Inner=Local, Inner=ZenLocal, Inner=Shared, Inner=Cloud)
+Local=(Type=FileSystem, PurgeTransient=true, DeleteOnly=true, UnusedFileAge=34, PromptIfMissing=true, Path=%ENGINEDIR%DerivedDataCache, EnvPathOverride=UE-LocalDataCachePath, EditorOverrideSetting=LocalDerivedDataCache, CommandLineOverride=LocalDataCachePath)
+ZenLocal=(Type=Zen, Namespace="ue.ddc")
 Shared=(Type=FileSystem, UnusedFileAge=10, FoldersToClean=10, ConsiderSlowAt=70, Path=?EpicDDC, EnvPathOverride=UE-SharedDataCachePath, EditorOverrideSetting=SharedDerivedDataCache, CommandLineOverride=SharedDataCachePath)
+Cloud=(Type=Cloud, ServerID=Default, EnvHostOverride=UE-CloudDataCacheHost, CommandLineHostOverride=CloudDataCacheHost, CommandLineOAuthSecretOverride="CloudDataCacheOAuthSecret", OAuthAccessTokenEnvOverride="UE-CloudDataCacheAccessToken")
 Pak=(Type=ReadPak, Filename="%GAMEDIR%DerivedDataCache/DDC.ddp")
 EnginePak=(Type=ReadPak, Filename="%ENGINEDIR%DerivedDataCache/DDC.ddp")
 
diff --git a/Engine/Source/Developer/DerivedDataCache/Private/ZenCacheStore.cpp b/Engine/Source/Developer/DerivedDataCache/Private/ZenCacheStore.cpp
--- a/Engine/Source/Developer/DerivedDataCache/Private/ZenCacheStore.cpp
+++ b/Engine/Source/Developer/DerivedDataCache/Private/ZenCacheStore.cpp
@@ -1204,19 +1204,13 @@ TTuple<ILegacyCacheStore*, ECacheStoreFlags> CreateZenCacheStore(const TCHAR* No
 	FParse::Value(Config, TEXT("Host="), ServiceUrl);
 
 	FString Namespace;
-	if (!FParse::Value(Config, TEXT("Namespace="), Namespace))
+	if (!FParse::Value(Config, TEXT("StructuredNamespace="), Namespace) && !FParse::Value(Config, TEXT("Namespace="), Namespace))
 	{
 		Namespace = FApp::GetProjectName();
+		UE_LOG(LogDerivedDataCache, Warning, TEXT("%s: Missing required parameter 'Namespace', falling back to '%s'"), NodeName, *Namespace);
 	}
 
-	FString StructuredNamespace;
-	if (!FParse::Value(Config, TEXT("StructuredNamespace="), StructuredNamespace))
-	{
-		StructuredNamespace = Namespace;
-		UE_LOG(LogDerivedDataCache, Warning, TEXT("%s: Missing required parameter 'StructuredNamespace', falling back '%s'"), NodeName, *StructuredNamespace);
-	}
-
-	TUniquePtr<FZenCacheStore> Backend = MakeUnique<FZenCacheStore>(*ServiceUrl, *StructuredNamespace);
+	TUniquePtr<FZenCacheStore> Backend = MakeUnique<FZenCacheStore>(*ServiceUrl, *Namespace);
 	if (!Backend->IsUsable())
 	{
 		UE_LOG(LogDerivedDataCache, Warning, TEXT("%s: Failed to contact the service (%s), will not use it."), NodeName, *Backend->GetName());
