Lumen - periodically refresh surface cache on all scalability levels

diff --git a/Engine/Config/BaseScalability.ini b/Engine/Config/BaseScalability.ini
--- a/Engine/Config/BaseScalability.ini
+++ b/Engine/Config/BaseScalability.ini
@@ -268,7 +268,6 @@ r.Lumen.TranslucencyVolume.GridPixelSize=64
 r.Lumen.TranslucencyVolume.TraceFromVolume=0
 r.Lumen.TranslucencyVolume.RadianceCache.ProbeResolution=8
 r.Lumen.TranslucencyVolume.RadianceCache.NumProbesToTraceBudget=200
-r.LumenScene.SurfaceCache.CardCaptureRefreshFraction=0
 
 [GlobalIlluminationQuality@3]
 r.DistanceFieldAO=1
@@ -292,7 +291,6 @@ r.Lumen.TranslucencyVolume.TraceFromVolume=1
 r.Lumen.TranslucencyVolume.TracingOctahedronResolution=3
 r.Lumen.TranslucencyVolume.RadianceCache.ProbeResolution=8
 r.Lumen.TranslucencyVolume.RadianceCache.NumProbesToTraceBudget=200
-r.LumenScene.SurfaceCache.CardCaptureRefreshFraction=0.125
 
 [GlobalIlluminationQuality@Cine]
 r.DistanceFieldAO=1
@@ -316,7 +314,6 @@ r.Lumen.TranslucencyVolume.TraceFromVolume=1
 r.Lumen.TranslucencyVolume.TracingOctahedronResolution=4
 r.Lumen.TranslucencyVolume.RadianceCache.ProbeResolution=16
 r.Lumen.TranslucencyVolume.RadianceCache.NumProbesToTraceBudget=1000
-r.LumenScene.SurfaceCache.CardCaptureRefreshFraction=0.125
 
 
 ;-----------------------------------------------------------------------------------------------------------------
