Workaround fix for certain Android devices not correctly flipping to portrait mode as expected. Setting all devices to use r.Android.GLESFlipYMethod=2 instead of the new =1 method.
#tests none

diff --git a/Engine/Config/BaseDeviceProfiles.ini b/Engine/Config/BaseDeviceProfiles.ini
--- a/Engine/Config/BaseDeviceProfiles.ini
+++ b/Engine/Config/BaseDeviceProfiles.ini
@@ -817,7 +817,7 @@ BaseProfileName=Android_Mid
 DeviceType=Android
 BaseProfileName=Android_High
 ; seems like all Adreno6xx devices properly support setBuffersTransform
-+CVars=r.Android.GLESFlipYMethod=1
++CVars=r.Android.GLESFlipYMethod=2
 
 [Android_Adreno6xx_Vulkan DeviceProfile]
 DeviceType=Android
