- Changed build default back to Intel so builds on builders work as expected

diff --git a/Engine/Config/BaseEngine.ini b/Engine/Config/BaseEngine.ini
--- a/Engine/Config/BaseEngine.ini
+++ b/Engine/Config/BaseEngine.ini
@@ -2852,8 +2852,8 @@ UseFastIntrinsics=False
 EnableMathOptimisations=True
 TargetArchitecture=MacTargetArchitectureUniversal
 EditorTargetArchitecture=MacTargetArchitectureUniversal
-EditorDefaultArchitecture=Host
-DefaultArchitecture=Host
+EditorDefaultArchitecture=MacTargetArchitectureIntel
+DefaultArchitecture=MacTargetArchitectureIntel
 
 
 [HMDPluginPriority]
