[Navmesh] Reverting change of the default CellSize recently introduced in CL 23434938
(changing the default cell size ends up releasing the DetourNavMesh during serialization and this has unexpected consequences)
#rnx

diff --git a/Engine/Config/BaseEngine.ini b/Engine/Config/BaseEngine.ini
--- a/Engine/Config/BaseEngine.ini
+++ b/Engine/Config/BaseEngine.ini
@@ -2501,10 +2501,10 @@ MaxTileGridHeight=256
 DefaultDrawDistance=5000.0
 ; navmesh generation parameters
 TileSizeUU=1000.f
-CellSize=20.f
-NavMeshResolutionParams[0]=(CellSize=40.000000)
-NavMeshResolutionParams[1]=(CellSize=20.000000)
-NavMeshResolutionParams[2]=(CellSize=10.000000)
+CellSize=19.f
+NavMeshResolutionParams[0]=(CellSize=38.000000)
+NavMeshResolutionParams[1]=(CellSize=19.000000)
+NavMeshResolutionParams[2]=(CellSize=19.000000)
 CellHeight=10.f
 AgentRadius=34.f
 AgentHeight=144.f
