Cooker: Remove CookPlatformDataCacheSettings default settings for materials and textures. This throttle to to SaveQueue processing is no longer necessary because materials and textures now have their own limiters for how many are in flight.
Removing the throttle allows the SaveQueue to go wider on DDC requests and reduces the size of the save queue, which reduces the cost of garbage collection.
#rnx

diff --git a/Engine/Config/BaseEditor.ini b/Engine/Config/BaseEditor.ini
--- a/Engine/Config/BaseEditor.ini
+++ b/Engine/Config/BaseEditor.ini
@@ -472,9 +472,9 @@ IterativeClassAllowList=/Script/Engine.Texture2D
 ; +IterativeClassAllowList=/Script/UnrealEd.FbxAnimSequenceImportData
 
 [CookPlatformDataCacheSettings]
-MaterialInstance=100
-Material=100
-Texture2D=50
+; Add <ShortClassName>=<Integer> here to limit the number of instances that can have BeginCacheForCookedPlatformData called at once
+; This was originally created for textures and materials, but we now have thread pools that handle them in a system-specific way
+;Texture2D=50
 
 [TextureBuild]
 NewTextureBuilds=false
