Removed deprecated input binding potentially causing flickering issue with the Gameplay Debugger menu.

#tests PIE testing pressing ' still toggling the debugger

diff --git a/Engine/Config/BaseInput.ini b/Engine/Config/BaseInput.ini
--- a/Engine/Config/BaseInput.ini
+++ b/Engine/Config/BaseInput.ini
@@ -47,9 +47,6 @@ bEnableInputDeviceSubsystem=true
 +DebugExecBindings=(Key=Tab,Command="FocusLastPIEWindow",Control=True,Shift=True)
 +DebugExecBindings=(Key=PageDown,Command="PreviousDebugTarget")
 +DebugExecBindings=(Key=PageUp,Command="NextDebugTarget")
-; Please remove EnableGDT command binding after removing deprecated GameplayDebugger module (we are going to use Gameplay Debugger plugin soon)
-+DebugExecBindings=(Key=Apostrophe,Command="EnableGDT")
-+DebugExecBindings=(Key=Quote,Command="EnableGDT")
 ; --- PIE only binding
 +DebugExecBindings=(Key=Semicolon,Command="ToggleDebugCamera")
 
