Adding 0.1 to the default snap options

diff --git a/Engine/Config/BaseEditorPerProjectUserSettings.ini b/Engine/Config/BaseEditorPerProjectUserSettings.ini
--- a/Engine/Config/BaseEditorPerProjectUserSettings.ini
+++ b/Engine/Config/BaseEditorPerProjectUserSettings.ini
@@ -407,6 +407,7 @@ bLevelEditorJoystickControls=True
 .ScalingGridSizes=0.5
 .ScalingGridSizes=0.25
 .ScalingGridSizes=0.125
+.ScalingGridSizes=0.1
 .ScalingGridSizes=0.0625
 .ScalingGridSizes=0.03125
 
