Rolls back r.TSR.History.GrandReprojection=1 on Epic TSR

diff --git a/Engine/Config/BaseScalability.ini b/Engine/Config/BaseScalability.ini
--- a/Engine/Config/BaseScalability.ini
+++ b/Engine/Config/BaseScalability.ini
@@ -42,6 +42,7 @@ r.TemporalAA.Quality=0
 r.TSR.History.R11G11B10=1
 r.TSR.History.ScreenPercentage=100
 r.TSR.History.UpdateQuality=0
+r.TSR.History.GrandReprojection=0
 r.TSR.ShadingRejection.Flickering=0
 r.TSR.Velocity.Extrapolation=0
 r.TSR.RejectionAntiAliasingQuality=0
@@ -52,6 +53,7 @@ r.TemporalAA.Quality=1
 r.TSR.History.R11G11B10=1
 r.TSR.History.ScreenPercentage=100
 r.TSR.History.UpdateQuality=1
+r.TSR.History.GrandReprojection=0
 r.TSR.ShadingRejection.Flickering=0
 r.TSR.Velocity.Extrapolation=0
 r.TSR.RejectionAntiAliasingQuality=1
@@ -62,6 +64,7 @@ r.TemporalAA.Quality=1
 r.TSR.History.R11G11B10=1
 r.TSR.History.ScreenPercentage=100
 r.TSR.History.UpdateQuality=2
+r.TSR.History.GrandReprojection=0
 r.TSR.ShadingRejection.Flickering=1
 r.TSR.Velocity.Extrapolation=0
 r.TSR.RejectionAntiAliasingQuality=1
@@ -70,8 +73,9 @@ r.TSR.RejectionAntiAliasingQuality=1
 r.FXAA.Quality=4
 r.TemporalAA.Quality=2
 r.TSR.History.R11G11B10=1
-r.TSR.History.ScreenPercentage=200
+r.TSR.History.ScreenPercentage=100
 r.TSR.History.UpdateQuality=3
+r.TSR.History.GrandReprojection=1
 r.TSR.ShadingRejection.Flickering=1
 r.TSR.Velocity.Extrapolation=1
 r.TSR.RejectionAntiAliasingQuality=2
@@ -82,6 +86,7 @@ r.TemporalAA.Quality=2
 r.TSR.History.R11G11B10=0
 r.TSR.History.ScreenPercentage=200
 r.TSR.History.UpdateQuality=3
+r.TSR.History.GrandReprojection=1
 r.TSR.ShadingRejection.Flickering=1
 r.TSR.Velocity.Extrapolation=1
 r.TSR.RejectionAntiAliasingQuality=2
