Only generate niagara sidecar data for the editor plugin instead of the whole client cook as well

diff --git a/Engine/Config/BaseEditor.ini b/Engine/Config/BaseEditor.ini
--- a/Engine/Config/BaseEditor.ini
+++ b/Engine/Config/BaseEditor.ini
@@ -383,11 +383,6 @@ CookProcessCount=1
 +CookOnTheFlyConfigSettingDenyList=*.Engine:/Script/UnrealEd.CookerSettings:bIterativeCookingForFileCookContent
 +CookOnTheFlyConfigSettingDenyList=*.Engine:OnlineSubsystemMcp
 
-; Asset type to automatically generate editor optional sidecar data for
-;+AutomaticOptionalInclusionAssetType=/Script/Niagara.NiagaraScript
-;+AutomaticOptionalInclusionAssetType=/Script/Niagara.NiagaraEmitter
-;+AutomaticOptionalInclusionAssetType=/Script/Niagara.NiagaraSystem
-
 [EditorDomain]
 EditorDomainEnabled=false
 SaveUnversioned=true
