Update recommended IDE install prompt to download and install VS2022 rather than VS2019

#rnx

diff --git a/Engine/Config/BaseEditor.ini b/Engine/Config/BaseEditor.ini
--- a/Engine/Config/BaseEditor.ini
+++ b/Engine/Config/BaseEditor.ini
@@ -276,7 +276,7 @@ ForumsURL="http://forums.unrealengine.com"
 SearchForAnswersURL="https://forums.unrealengine.com/tag/question"
 SnippetsURL="https://dev.epicgames.com/community/snippets"
 SupportWebsiteURL="http://unrealengine.com/support"
-SourceCodeIDEURL_Windows="https://aka.ms/vs/16/release/vs_Community.exe"
+SourceCodeIDEURL_Windows="https://aka.ms/vs/17/release/vs_Community.exe"
 SourceCodeIDEURL_Mac="https://developer.apple.com/xcode"
 SourceCodeIDEURL_Linux="https://code.visualstudio.com/docs/setup/linux"
 SourceCodeIDEURL_Other="http://en.wikipedia.org/wiki/Integrated_development_environment"
diff --git a/Engine/Source/Editor/UnrealEd/Private/SourceCodeNavigation.cpp b/Engine/Source/Editor/UnrealEd/Private/SourceCodeNavigation.cpp
--- a/Engine/Source/Editor/UnrealEd/Private/SourceCodeNavigation.cpp
+++ b/Engine/Source/Editor/UnrealEd/Private/SourceCodeNavigation.cpp
@@ -1805,7 +1805,7 @@ FText FSourceCodeNavigation::GetSuggestedSourceCodeIDE(bool bShortIDEName)
 	}
 	else
 	{
-		return LOCTEXT("SuggestedCodeIDE_Windows", "Visual Studio 2019");
+		return LOCTEXT("SuggestedCodeIDE_Windows", "Visual Studio 2022");
 	}
 #elif PLATFORM_MAC
 	return LOCTEXT("SuggestedCodeIDE_Mac", "Xcode");
