change BaseEditor config for HDR texture import to BC6
TextureImportSettings CompressedFormatForFloatTextures=HDRCompressed_BC6

diff --git a/Engine/Config/BaseEditor.ini b/Engine/Config/BaseEditor.ini
--- a/Engine/Config/BaseEditor.ini
+++ b/Engine/Config/BaseEditor.ini
@@ -618,3 +618,8 @@ bCanModifyPluginsFromBrowser=true
 
 [EditorBulkData]
 LogCookedStatus=false
+
+
+[/Script/TextureUtilitiesCommon.TextureImportSettings]
+CompressedFormatForFloatTextures=HDRCompressed_BC6
+
