Niagara - Move the cooked sidecar optional data generation back to the base editor ini until the Content Worker issues can be resolved.

diff --git a/Engine/Config/BaseEditor.ini b/Engine/Config/BaseEditor.ini
--- a/Engine/Config/BaseEditor.ini
+++ b/Engine/Config/BaseEditor.ini
@@ -383,6 +383,13 @@ CookProcessCount=1
 +CookOnTheFlyConfigSettingDenyList=*.Engine:/Script/UnrealEd.CookerSettings:bIterativeCookingForFileCookContent
 +CookOnTheFlyConfigSettingDenyList=*.Engine:OnlineSubsystemMcp
 
+; Asset type to automatically generate editor optional sidecar data for
++AutomaticOptionalInclusionAssetType=/Script/Niagara.NiagaraScript
++AutomaticOptionalInclusionAssetType=/Script/Niagara.NiagaraEmitter
++AutomaticOptionalInclusionAssetType=/Script/Niagara.NiagaraSystem
++AutomaticOptionalInclusionAssetType=/Script/Niagara.NiagaraParameterDefinitions
++AutomaticOptionalInclusionAssetType=/Script/Niagara.NiagaraEffectType
+
 [EditorDomain]
 EditorDomainEnabled=false
 SaveUnversioned=true
