Niagara - Fix incorrect type namespace in previous ini checkin.

diff --git a/Engine/Config/BaseEditor.ini b/Engine/Config/BaseEditor.ini
--- a/Engine/Config/BaseEditor.ini
+++ b/Engine/Config/BaseEditor.ini
@@ -387,7 +387,7 @@ CookProcessCount=1
 +AutomaticOptionalInclusionAssetType=/Script/Niagara.NiagaraScript
 +AutomaticOptionalInclusionAssetType=/Script/Niagara.NiagaraEmitter
 +AutomaticOptionalInclusionAssetType=/Script/Niagara.NiagaraSystem
-+AutomaticOptionalInclusionAssetType=/Script/Niagara.NiagaraParameterDefinitions
++AutomaticOptionalInclusionAssetType=/Script/NiagaraEditor.NiagaraParameterDefinitions
 +AutomaticOptionalInclusionAssetType=/Script/Niagara.NiagaraEffectType
 
 [EditorDomain]
