Commits to r.TSR.History.ScreenPercentage=200 on Epic TSR scalability

diff --git a/Engine/Config/BaseScalability.ini b/Engine/Config/BaseScalability.ini
--- a/Engine/Config/BaseScalability.ini
+++ b/Engine/Config/BaseScalability.ini
@@ -73,9 +73,9 @@ r.TSR.RejectionAntiAliasingQuality=1
 r.FXAA.Quality=4
 r.TemporalAA.Quality=2
 r.TSR.History.R11G11B10=1
-r.TSR.History.ScreenPercentage=100
+r.TSR.History.ScreenPercentage=200
 r.TSR.History.UpdateQuality=3
-r.TSR.History.GrandReprojection=1
+r.TSR.History.GrandReprojection=0
 r.TSR.ShadingRejection.Flickering=1
 r.TSR.Velocity.Extrapolation=1
 r.TSR.RejectionAntiAliasingQuality=2
