Iris
- Iris config changes
- Renamed  IrisNetDriverConfig.bEnableIris to bCanUseIris to clarify what it does

diff --git a/Engine/Config/BaseEngine.ini b/Engine/Config/BaseEngine.ini
--- a/Engine/Config/BaseEngine.ini
+++ b/Engine/Config/BaseEngine.ini
@@ -297,8 +297,11 @@ NetClientTicksPerSecond=200
 +StatColorMappings=(StatName="Streaming fudge factor",ColorMap=((In=0.0,Out=(G=255)),(In=1.0,Out=(G=255)),(In=2.5,Out=(R=255,G=255)),(In=5.0,Out=(R=255)),(In=10.0,Out=(R=255))))
 DisplayGamma=2.2
 MinDesiredFrameRate=35.000000
+; NetDriver configs
 NetDriverDefinitions=(DefName="GameNetDriver",DriverClassName="/Script/OnlineSubsystemUtils.IpNetDriver",DriverClassNameFallback="/Script/OnlineSubsystemUtils.IpNetDriver")
 +NetDriverDefinitions=(DefName="DemoNetDriver",DriverClassName="/Script/Engine.DemoNetDriver",DriverClassNameFallback="/Script/Engine.DemoNetDriver")
++IrisNetDriverConfigs=(NetDriverDefinition=GameNetDriver, bCanUseIris=true)
++IrisNetDriverConfigs=(NetDriverName=DemoNetDriver, bCanUseIris=false)
 NetErrorLogInterval=1.0
 SerializationOutOfBoundsErrorMessage=NSLOCTEXT("","SerializationOutOfBoundsErrorMessage","Corrupt data found, please verify your installation.")
 SerializationOutOfBoundsErrorMessageCaption=NSLOCTEXT("","SerializationOutOfBoundsErrorMessageCaption","Serialization Error : Action Needed")
@@ -1163,6 +1166,67 @@ UIScaleCurve=(EditorCurveData=(Keys=((Time=480,Value=0.444),(Time=720,Value=0.66
 bLoadWidgetsOnDedicatedServer=True
 bAllowHighDPIInGameMode=False
 
+;Iris - begin
+
+[/Script/IrisCore.PartialNetObjectAttachmentHandlerConfig]
+MaxPartCount=4096
+
+[/Script/IrisCore.NetBlobHandlerDefinitions]
++NetBlobHandlerDefinitions=(ClassName=NetRPCHandler)
++NetBlobHandlerDefinitions=(ClassName=PartialNetObjectAttachmentHandler )
++NetBlobHandlerDefinitions=(ClassName=NetObjectBlobHandler)
+
+[/Script/IrisCore.DataStreamDefinitions]
++DataStreamDefinitions=(DataStreamName=NetToken, ClassName=/Script/IrisCore.NetTokenDataStream, DefaultSendStatus=EDataStreamSendStatus::Send, bAutoCreate=true)
++DataStreamDefinitions=(DataStreamName=Replication, ClassName=/Script/IrisCore.ReplicationDataStream, DefaultSendStatus=EDataStreamSendStatus::Send, bAutoCreate=true)
+
+[/Script/IrisCore.NetObjectPrioritizerDefinitions]
++NetObjectPrioritizerDefinitions=(PrioritizerName=Default, ClassName=/Script/IrisCore.SphereNetObjectPrioritizer, ConfigClassName=/Script/IrisCore.SphereNetObjectPrioritizerConfig)
++NetObjectPrioritizerDefinitions=(PrioritizerName=PlayerState, ClassName=/Script/IrisCore.NetObjectCountLimiter, ConfigClassName=/Script/Engine.PlayerStateCountLimiterConfig)
+
+[/Script/IrisCore.NetObjectFilterDefinitions]
++NetObjectFilterDefinitions=(FilterName=Spatial, ClassName=/Script/IrisCore.NetObjectGridFilter, ConfigClassName=/Script/IrisCore.NetObjectGridFilterConfig)
++NetObjectFilterDefinitions=(FilterName=NotRouted, ClassName=/Script/IrisCore.FilterOutNetObjectFilter, ConfigClassName=/Script/IrisCore.FilterOutNetObjectFilterConfig)
+
+[/Script/IrisCore.ObjectReplicationBridgeConfig]
+; It is possible to override how often replicated objects are Polled for dirtyness if they are not using pushed based dirtiness, normally based on NetUpdateFrequency
+;+PollConfigs=(ClassName=/Script/Example.Pawn, PollFramePeriod=0, bIncludeSubclasses=true)
+
+; Filters
+DefaultSpatialFilterName=Spatial
+; Filter configs should typically be configured per game
+; Clear all filters
+!FilterConfigs=ClearArray
+;+FilterConfigs=(ClassName=/Script/Engine.LevelScriptActor, DynamicFilterName=NotRouted) ; NotRouted means that this type will not be replicated
++FilterConfigs=(ClassName=/Script/Engine.Actor, DynamicFilterName=None))
+; Info types aren't supposed to have physical representation and rely on a static priority
++FilterConfigs=(ClassName=/Script/Engine.Info, DynamicFilterName=None)
++FilterConfigs=(ClassName=/Script/Engine.PlayerState, DynamicFilterName=None)
+; Pawns can be spatially filtered...
+;+FilterConfigs=(ClassName=/Script/Engine.Pawn, DynamicFilterName=Spatial))
+
+; PrioritizerConfigs
++PrioritizerConfigs=(ClassName=/Script/Engine.PlayerState, PrioritizerName=PlayerState)
+
+;DeltaCompressionConfigs=ClearArray
+; Classes for which deltacompression is enabled
++DeltaCompressionConfigs=(ClassName=/Script/Engine.Pawn))
++DeltaCompressionConfigs=(ClassName=/Script/Engine.PlayerState))
+
+[/Script/IrisCore.ReplicationStateDescriptorConfig]
+; Declarate structs that are vetted to work using reflection based struct serialization even though there exists a custom NetSerialize function for the struct
++SupportsStructNetSerializerList=(StructName=GameplayCueParameters)
++SupportsStructNetSerializerList=(StructName=GameplayAbilityTargetData_LocationInfo)
++SupportsStructNetSerializerList=(StructName=GameplayAbilityTargetData_ActorArray)
++SupportsStructNetSerializerList=(StructName=GameplayAbilityTargetData_SingleTargetHit)
++SupportsStructNetSerializerList=(StructName=NetLevelVisibilityTransactionId)
++SupportsStructNetSerializerList=(StructName=Vector2D)
++SupportsStructNetSerializerList=(StructName=GameplayDebuggerNetPack)
+
+RequiredNetDriverChannelClassName=/Script/Engine.DataStreamChannel
+
+;Iris - end
+
 [/Script/Engine.GameEngine]
 MaxDeltaTime=0
 ServerFlushLogInterval=30
@@ -1415,6 +1479,7 @@ ConfiguredLanSpeed=100000
 [/Script/Engine.NetDriver]
 +ChannelDefinitions=(ChannelName=Control, ClassName=/Script/Engine.ControlChannel, StaticChannelIndex=0, bTickOnCreate=true, bServerOpen=false, bClientOpen=true, bInitialServer=false, bInitialClient=true)
 +ChannelDefinitions=(ChannelName=Voice, ClassName=/Script/Engine.VoiceChannel, StaticChannelIndex=1, bTickOnCreate=true, bServerOpen=true, bClientOpen=true, bInitialServer=true, bInitialClient=true)
++ChannelDefinitions=(ChannelName=DataStream, ClassName=/Script/Engine.DataStreamChannel, StaticChannelIndex=2, bTickOnCreate=true, bServerOpen=true, bClientOpen=true, bInitialServer=true, bInitialClient=true)
 +ChannelDefinitions=(ChannelName=Actor, ClassName=/Script/Engine.ActorChannel, StaticChannelIndex=-1, bTickOnCreate=false, bServerOpen=true, bClientOpen=false, bInitialServer=false, bInitialClient=false)
 
 [/Script/Engine.RendererSettings]
diff --git a/Engine/Source/Runtime/Engine/Classes/Engine/Engine.h b/Engine/Source/Runtime/Engine/Classes/Engine/Engine.h
--- a/Engine/Source/Runtime/Engine/Classes/Engine/Engine.h
+++ b/Engine/Source/Runtime/Engine/Classes/Engine/Engine.h
@@ -275,10 +275,10 @@ struct FIrisNetDriverConfig
 	FString NetDriverWildcardName;
 
 	/**
-	 * Configurable property that decides if the NetDriver will use the Iris replication system or not
+	 * Configurable property that decides if the NetDriver will use the Iris replication system or not if Iris is enabled
 	 */
 	UPROPERTY()
-	bool bEnableIris = false;
+	bool bCanUseIris = false;
 };
 
 
diff --git a/Engine/Source/Runtime/Engine/Private/UnrealEngine.cpp b/Engine/Source/Runtime/Engine/Private/UnrealEngine.cpp
--- a/Engine/Source/Runtime/Engine/Private/UnrealEngine.cpp
+++ b/Engine/Source/Runtime/Engine/Private/UnrealEngine.cpp
@@ -13387,12 +13387,12 @@ UNetDriver* UEngine::FindNamedNetDriver(const UPendingNetGame* InPendingNetGame,
 namespace UE::Private
 {
 	/**
-	 * Look for a config setting telling if this new NetDriver should be enabled with Iris or not
+	 * Look for a config setting telling if this new NetDriver can use Iris or not
 	 * To enable iris for a netdriver you need to add the proper configuration in Engine.ini like so:
 	 *		[/Script/Engine.Engine]
-	 *		+IrisNetDriverConfigs=(NetDriverName="MyNetDriver",bEnableIris=true)
-	 *		+IrisNetDriverConfigs=(NetDriverDefinition="MyNetDriverDef",bEnableIris=false)
-	 *		+IrisNetDriverConfigs=(NetDriverWildcardName="SecondaryNetDriver*",bEnableIris=true);
+	 *		+IrisNetDriverConfigs=(NetDriverName="MyNetDriver",bCanUseIris=true)
+	 *		+IrisNetDriverConfigs=(NetDriverDefinition="MyNetDriverDef",bCanUseIris=false)
+	 *		+IrisNetDriverConfigs=(NetDriverWildcardName="SecondaryNetDriver*",bCanUseIris=true);
 	 *
 	 * Priority order for the IrisNetDriverConfigs are:
 	 *		1. NetDriverName exact match
@@ -13440,7 +13440,7 @@ namespace UE::Private
 		}
 
 		// Only use Iris if the module is loaded (this happens automatically if the Iris plugin is enabled)
-		if (IrisConfig && IrisConfig->bEnableIris)
+		if (IrisConfig && IrisConfig->bCanUseIris)
 		{
 			if (!ensureMsgf(FModuleManager::Get().IsModuleLoaded("IrisCore"), TEXT("%s is not using Iris because the IrisCore module isn't loaded. Check whether the Iris plugin is enabled."), *InNetDriverName.ToString()))
 			{
@@ -13448,7 +13448,7 @@ namespace UE::Private
 			}
 		}
 
-		return IrisConfig ? IrisConfig->bEnableIris : false;
+		return IrisConfig ? IrisConfig->bCanUseIris : false;
 #else
 		return false;
 #endif //UE_WITH_IRIS
