Iris
- Marked struct FVector2f as safe to serialize using reflection data to silence warning

diff --git a/Engine/Config/BaseEngine.ini b/Engine/Config/BaseEngine.ini
--- a/Engine/Config/BaseEngine.ini
+++ b/Engine/Config/BaseEngine.ini
@@ -1221,6 +1221,7 @@ DefaultSpatialFilterName=Spatial
 +SupportsStructNetSerializerList=(StructName=GameplayAbilityTargetData_SingleTargetHit)
 +SupportsStructNetSerializerList=(StructName=NetLevelVisibilityTransactionId)
 +SupportsStructNetSerializerList=(StructName=Vector2D)
++SupportsStructNetSerializerList=(StructName=Vector2f)
 +SupportsStructNetSerializerList=(StructName=GameplayDebuggerNetPack)
 
 RequiredNetDriverChannelClassName=/Script/Engine.DataStreamChannel
