Iris - Silence warnings when building replication descriptors for some RootMotionSource structs. They've been verified good at this time.

diff --git a/Engine/Config/BaseEngine.ini b/Engine/Config/BaseEngine.ini
--- a/Engine/Config/BaseEngine.ini
+++ b/Engine/Config/BaseEngine.ini
@@ -1223,6 +1223,12 @@ DefaultSpatialFilterName=Spatial
 +SupportsStructNetSerializerList=(StructName=Vector2D)
 +SupportsStructNetSerializerList=(StructName=Vector2f)
 +SupportsStructNetSerializerList=(StructName=GameplayDebuggerNetPack)
++SupportsStructNetSerializerList=(StructName=RootMotionSource)
++SupportsStructNetSerializerList=(StructName=RootMotionSource_ConstantForce)
++SupportsStructNetSerializerList=(StructName=RootMotionSource_JumpForce)
++SupportsStructNetSerializerList=(StructName=RootMotionSource_MoveToDynamicForce)
++SupportsStructNetSerializerList=(StructName=RootMotionSource_MoveToForce)
++SupportsStructNetSerializerList=(StructName=RootMotionSource_RadialForce)
 
 RequiredNetDriverChannelClassName=/Script/Engine.DataStreamChannel
 
