Fix incorrect regex for double digit iPhone model ids. For example, the previous profile "iPad13,[16-17]=iPadAir5" will not match an iOS Model "iPad13,16" - this fixes the invalid iPad13,xy matches.

diff --git a/Engine/Config/BaseDeviceProfiles.ini b/Engine/Config/BaseDeviceProfiles.ini
--- a/Engine/Config/BaseDeviceProfiles.ini
+++ b/Engine/Config/BaseDeviceProfiles.ini
@@ -215,12 +215,14 @@ iPad6,[3,4]=iPadPro97
 iPad6,[7,8]=iPadPro129
 iPad6,11=iPad5
 iPad6,12=iPad5
-iPad7,[1,2]=iPadPro2_129
+iPad7,1$=iPadPro2_129
+iPad7,2=iPadPro2_129
 iPad7,[3,4]=iPadPro105
 iPad7,[5,6]=iPad6
 iPad7,11=iPad7
 iPad7,12=iPad7
-iPad8,[1-4]=iPadPro11
+iPad8,1$=iPadPro11
+iPad8,[2-4]=iPadPro11
 iPad8,[5-8]=iPadPro3_129
 iPad8,9=iPadPro2_11
 iPad8,10=iPadPro2_11
@@ -230,11 +232,13 @@ iPad11,[1,2]=iPadMini5
 iPad11,[3,4]=iPadAir3
 iPad11,[6,7]=iPad8
 iPad12,[1,2]=iPad9
-iPad13,[1,2]=iPadAir4
+iPad13,1$=iPadAir4
+iPad13,2=iPadAir4
 iPad13,[4-7]=iPadPro3_11
-iPad13,[8-11]=iPadPro5_129
-iPad13,[16-17]=iPadAir5
-iPad13,[18-19]=iPad10
+iPad13,[8-9]=iPadPro5_129
+iPad13,1[0-1]=iPadPro5_129
+iPad13,1[6-7]=iPadAir5
+iPad13,1[8-9]=iPad10
 iPad14,[1,2]=iPadMini6
 iPad14,[3,4]=iPadPro4_11
 iPad14,[5,6]=iPadPro6_129
