Clean up older VA config file options

#rnx

- Add quotes around the 'VirtualizationProcessTag' entry
- Remove an old graph that stuff referenced jupiter

diff --git a/Engine/Config/BaseEngine.ini b/Engine/Config/BaseEngine.ini
--- a/Engine/Config/BaseEngine.ini
+++ b/Engine/Config/BaseEngine.ini
@@ -1341,7 +1341,7 @@ EnableCacheOnPull=True
 EnableCacheOnPush=True
 MinPayloadLength=0
 BackendGraph=ContentVirtualizationBackendGraph_None
-VirtualizationProcessTag=#virtualized
+VirtualizationProcessTag="#virtualized"
 AllowSubmitIfVirtualizationFailed=false
 FilterMode=OptOut
 FilterMapContent=True
@@ -1363,14 +1363,7 @@ PersistentStorageHierarchy=()
 [ContentVirtualizationBackendGraph_LocalCache]
 CacheStorageHierarchy=(Entry=LocalFileCache)
 PersistentStorageHierarchy=()
-LocalFileCache=(Type=FileSystem, Path="%GAMESAVEDDIR%VirtualizedPayloads/")
-
-; A backend graph that will connect to a locally hosted Jupiter service along with a local file cache, suitable for testing.
-[ContentVirtualizationBackendGraph_JupiterLocalHost]
-CacheStorageHierarchy=(Entry=LocalFileCache)
-PersistentStorageHierarchy=(JupiterCache)
-LocalFileCache=(Type=FileSystem, Path="%GAMESAVEDDIR%VirtualizedPayloads/")
-JupiterCache=(Type=Jupiter, Host="http://localhost", Namespace="ue.mirage", ChunkSize=524288)
+LocalFileCache=(Type=FileSystem)
 
 [/Script/Engine.StreamingSettings]
 s.MinBulkDataSizeForAsyncLoading=131072
