Take Recorder: Disable clean view for default take recorder settings

#rnx

diff --git a/Engine/Config/BaseEditorPerProjectUserSettings.ini b/Engine/Config/BaseEditorPerProjectUserSettings.ini
--- a/Engine/Config/BaseEditorPerProjectUserSettings.ini
+++ b/Engine/Config/BaseEditorPerProjectUserSettings.ini
@@ -867,6 +867,7 @@ bInfiniteKeyAreas=true
 bAutoSetTrackDefaults=true
 FrameNumberDisplayFormat=NonDropFrameTimecode
 bAutoScrollEnabled=true
+bCleanPlaybackMode=false
 
 [EmbeddedActorSequenceEditor SequencerSettings]
 bShowRangeSlider=true
