VulkanRHI: Update driver version requirement for Vulkan on AMD in Windows to 23.3.2 which fixes crash in pipeline compilation.

#rnx

diff --git a/Engine/Config/BaseHardware.ini b/Engine/Config/BaseHardware.ini
--- a/Engine/Config/BaseHardware.ini
+++ b/Engine/Config/BaseHardware.ini
@@ -58,10 +58,13 @@ SuggestedDriverVersion="515.48"
 ;-----------------------------------------------------------------------------------------------------------------
 
 [GPU_AMD Windows]
-SuggestedDriverVersion="22.2.2"
+SuggestedDriverVersion="23.3.2"
 
 ; To prevent problems with older drivers.
-+DriverDenyList=(DriverVersion="<30.0.14011.2006", Reason="These driver versions have known stability issues and missing features")
++DriverDenyList=(DriverVersion="<30.0.14011.2006", Reason="These driver versions have known stability issues and missing features", RHI="D3D11")
++DriverDenyList=(DriverVersion="<30.0.14011.2006", Reason="These driver versions have known stability issues and missing features", RHI="D3D12")
+; 23.3.2 is necessary to prevent pipeline compilation crash in Vulkan
++DriverDenyList=(DriverVersion="<31.0.14043.1000", Reason="These driver versions have known stability issues and missing features", RHI="Vulkan")
 
 ;-----------------------------------------------------------------------------------------------------------------
 
