Changes r.Bloom.ScreenPercentage to 50 on epic post processing scalability

diff --git a/Engine/Config/BaseScalability.ini b/Engine/Config/BaseScalability.ini
--- a/Engine/Config/BaseScalability.ini
+++ b/Engine/Config/BaseScalability.ini
@@ -471,7 +471,7 @@ r.LensFlareQuality=2
 r.SceneColorFringeQuality=1
 r.EyeAdaptationQuality=2
 r.BloomQuality=5
-r.Bloom.ScreenPercentage=70.711
+r.Bloom.ScreenPercentage=50.000
 r.FastBlurThreshold=100
 r.Upscale.Quality=3
 r.Tonemapper.GrainQuantization=1
