Add a default TVOS DeviceProfile to allow Niagara to correctly cook required content. This allows Niagara CPU Particles to appear on tvOS devices.
Note: Niagara GPU Particles are not supported on Apple A8 devices (eg AppleTV HD) due to missing support for IndirectDraw in Metal

diff --git a/Engine/Config/BaseDeviceProfiles.ini b/Engine/Config/BaseDeviceProfiles.ini
--- a/Engine/Config/BaseDeviceProfiles.ini
+++ b/Engine/Config/BaseDeviceProfiles.ini
@@ -67,7 +67,7 @@
 +DeviceProfileNameAndTypes=AppleTV,IOS
 +DeviceProfileNameAndTypes=AppleTV4K,IOS
 +DeviceProfileNameAndTypes=AppleTV2_4K,IOS
-+DeviceProfileNameAndTypes=TVOS,AppleTV
++DeviceProfileNameAndTypes=TVOS,TVOS
 +DeviceProfileNameAndTypes=MacEditor,Mac
 +DeviceProfileNameAndTypes=MacClient,Mac
 +DeviceProfileNameAndTypes=Mac,Mac
@@ -619,6 +619,10 @@ BaseProfileName=IOS
 +CVars=r.Mobile.AmbientOcclusionQuality=0
 +CVars=r.Mobile.PixelProjectedReflectionQuality=0
 
+[TVOS DeviceProfile]
+DeviceType=TVOS
+BaseProfileName=AppleTV
+
 [AppleTV4K DeviceProfile]
 DeviceType=IOS
 BaseProfileName=AppleTV
