Add device profile for newly-release Mali-G7xx devices such as the Mali-G715 Immortalis

diff --git a/Engine/Config/BaseDeviceProfiles.ini b/Engine/Config/BaseDeviceProfiles.ini
--- a/Engine/Config/BaseDeviceProfiles.ini
+++ b/Engine/Config/BaseDeviceProfiles.ini
@@ -108,6 +108,8 @@
 +DeviceProfileNameAndTypes=Android_Mali_G78_Vulkan,Android
 +DeviceProfileNameAndTypes=Android_Mali_G710,Android
 +DeviceProfileNameAndTypes=Android_Mali_G710_Vulkan,Android
++DeviceProfileNameAndTypes=Android_Mali_G7xx,Android
++DeviceProfileNameAndTypes=Android_Mali_G7xx_Vulkan,Android
 +DeviceProfileNameAndTypes=Android_Xclipse_920,Android
 +DeviceProfileNameAndTypes=Android_Xclipse_920_Vulkan,Android
 +DeviceProfileNameAndTypes=Android_Vulkan_SM5,Android
@@ -672,6 +674,8 @@ BaseProfileName=AppleTV4K
 +MatchProfile=(Profile="Android_Mali_G78",Match=((SourceType=SRC_GpuFamily,CompareType=CMP_Regex,MatchString="^Mali\\-G78")))
 +MatchProfile=(Profile="Android_Mali_G710_Vulkan",Match=((SourceType=SRC_GpuFamily,CompareType=CMP_Regex,MatchString="^Mali\\-G710"),(SourceType=SRC_VulkanAvailable,CompareType=CMP_Equal,MatchString="true")))
 +MatchProfile=(Profile="Android_Mali_G710",Match=((SourceType=SRC_GpuFamily,CompareType=CMP_Regex,MatchString="^Mali\\-G710")))
++MatchProfile=(Profile="Android_Mali_G7xx_Vulkan",Match=((SourceType=SRC_GpuFamily,CompareType=CMP_Regex,MatchString="^Mali\\-G7[0-9][0-9]"),(SourceType=SRC_VulkanAvailable,CompareType=CMP_Equal,MatchString="true")))
++MatchProfile=(Profile="Android_Mali_G7xx",Match=((SourceType=SRC_GpuFamily,CompareType=CMP_Regex,MatchString="^Mali\\-G7[0-9][0-9]")))
 ; Enable Vulkan on any unknown Mali running Android 10 or later
 +MatchProfile=(Profile="Android_Unknown_Vulkan",Match=((SourceType=SRC_GpuFamily,CompareType=CMP_Regex,MatchString="^Mali"),(SourceType=SRC_AndroidVersion, CompareType=CMP_Regex,MatchString="([0-9]+).*"),(SourceType=SRC_PreviousRegexMatch,CompareType=CMP_GreaterEqual,MatchString="10"),(SourceType=SRC_VulkanAvailable,CompareType=CMP_Equal,MatchString="true")))
 
@@ -961,6 +965,16 @@ BaseProfileName=Android_Mali_G710
 ; enable Vulkan
 +CVars=r.Android.DisableVulkanSupport=0
 
+[Android_Mali_G7xx DeviceProfile]
+DeviceType=Android
+BaseProfileName=Android_Mali_G710
+
+[Android_Mali_G7xx_Vulkan DeviceProfile]
+DeviceType=Android
+BaseProfileName=Android_Mali_G7xx
+; enable Vulkan
++CVars=r.Android.DisableVulkanSupport=0
+
 [Android_TegraK1 DeviceProfile]
 DeviceType=Android
 BaseProfileName=Android_High
