BaseEngine.ini remove vestigial config option bForceAllBC23ToBC7

diff --git a/Engine/Config/BaseEngine.ini b/Engine/Config/BaseEngine.ini
--- a/Engine/Config/BaseEngine.ini
+++ b/Engine/Config/BaseEngine.ini
@@ -3259,8 +3259,6 @@ UnrealBuildTool=Engine/Build/BatchFiles/Build.bat
 TextureCompressionFormat="TextureFormatOodle"
 
 [TextureFormatOodleSettings]
-; should BC2 and 3 (DXT3 and 5) be changed to BC7
-bForceAllBC23ToBC7=False
 ; bDebugColor shows the BC1-7 format used
 bDebugColor=False
 ; DebugDump writes out the input and encoded images to Saved/OodleDebugImages, 
