Remove libc++_shared.so from apk packages on Android by default.
Enable advanced binary compression on Android

diff --git a/Engine/Config/BaseEngine.ini b/Engine/Config/BaseEngine.ini
--- a/Engine/Config/BaseEngine.ini
+++ b/Engine/Config/BaseEngine.ini
@@ -2951,10 +2951,10 @@ bAllowPatchOBBFile=False
 bAllowOverflowOBBFiles=False
 bExtractNativeLibs=true
 bEnableDomStorage=False
-bSkipLibCpp=False
+bSkipLibCpp=True
 bDisableLibCppSharedDependencyValidation=False
 bDisableFunctionDataSections=False
-bEnableAdvancedBinaryCompression=False
+bEnableAdvancedBinaryCompression=True
 bDisableStackProtector=False
 bEnableMulticastSupport=False
 
