Fixed class redirection directives

diff --git a/Engine/Config/BaseEngine.ini b/Engine/Config/BaseEngine.ini
--- a/Engine/Config/BaseEngine.ini
+++ b/Engine/Config/BaseEngine.ini
@@ -1152,39 +1152,39 @@ DefaultWorldPartitionHLODLayer=/Engine/Maps/Templates/HLODs/HLODLayer_Instanced.
 +ClassRedirects=(OldName="XRLoadingScreenFunctionLibrary", NewName="/Script/XRBase.XRLoadingScreenFunctionLibrary")
 
 ; Rename MaterialX Material Expressions and Enums
-+ClassRedirects=(OldName="/Script/InterchangeImport.UMaterialExpressionAppend3Vector", NewName="/Script/InterchangeImport.UMaterialExpressionMaterialXAppend3Vector")
-+ClassRedirects=(OldName="/Script/InterchangeImport.UMaterialExpressionAppend4Vector", NewName="/Script/InterchangeImport.UMaterialExpressionMaterialXAppend4Vector")
-+ClassRedirects=(OldName="/Script/InterchangeImport.UMaterialExpressionBurn", NewName="/Script/InterchangeImport.UMaterialExpressionMaterialXBurn")
-+ClassRedirects=(OldName="/Script/InterchangeImport.UMaterialExpressionDifference", NewName="/Script/InterchangeImport.UMaterialExpressionMaterialXDifference")
-+ClassRedirects=(OldName="/Script/InterchangeImport.UMaterialExpressionDisjointOver", NewName="/Script/InterchangeImport.UMaterialExpressionMaterialXDisjointOver")
-+ClassRedirects=(OldName="/Script/InterchangeImport.UMaterialExpressionDodge", NewName="/Script/InterchangeImport.UMaterialExpressionMaterialXDodge")
-+ClassRedirects=(OldName="/Script/InterchangeImport.UMaterialExpressionExponential", NewName="/Script/InterchangeImport.UMaterialExpressionMaterialXExponential")
-+ClassRedirects=(OldName="/Script/InterchangeImport.UMaterialExpressionFractal3D", NewName="/Script/InterchangeImport.UMaterialExpressionMaterialXFractal3D")
-+ClassRedirects=(OldName="/Script/InterchangeImport.UMaterialExpressionHsvToRgb", NewName="/Script/InterchangeImport.UMaterialExpressionMaterialXHsvToRgb")
-+ClassRedirects=(OldName="/Script/InterchangeImport.UMaterialExpressionIn", NewName="/Script/InterchangeImport.UMaterialExpressionMaterialXIn")
-+ClassRedirects=(OldName="/Script/InterchangeImport.UMaterialExpressionLength", NewName="/Script/InterchangeImport.UMaterialExpressionMaterialXLength")
-+ClassRedirects=(OldName="/Script/InterchangeImport.UMaterialExpressionLogarithm", NewName="/Script/InterchangeImport.UMaterialExpressionMaterialXLogarithm")
-+ClassRedirects=(OldName="/Script/InterchangeImport.UMaterialExpressionLuminance", NewName="/Script/InterchangeImport.UMaterialExpressionMaterialXLuminance")
-+ClassRedirects=(OldName="/Script/InterchangeImport.UMaterialExpressionMask", NewName="/Script/InterchangeImport.UMaterialExpressionMaterialXMask")
-+ClassRedirects=(OldName="/Script/InterchangeImport.UMaterialExpressionMatte", NewName="/Script/InterchangeImport.UMaterialExpressionMaterialXMatte")
-+ClassRedirects=(OldName="/Script/InterchangeImport.UMaterialExpressionMinus", NewName="/Script/InterchangeImport.UMaterialExpressionMaterialXMinus")
-+ClassRedirects=(OldName="/Script/InterchangeImport.UMaterialExpressionOut", NewName="/Script/InterchangeImport.UMaterialExpressionMaterialXOut")
-+ClassRedirects=(OldName="/Script/InterchangeImport.UMaterialExpressionOver", NewName="/Script/InterchangeImport.UMaterialExpressionMaterialXOver")
-+ClassRedirects=(OldName="/Script/InterchangeImport.UMaterialExpressionOverlay", NewName="/Script/InterchangeImport.UMaterialExpressionMaterialXOverlay")
-+ClassRedirects=(OldName="/Script/InterchangeImport.UMaterialExpressionPlace2D", NewName="/Script/InterchangeImport.UMaterialExpressionMaterialXPlace2D")
-+ClassRedirects=(OldName="/Script/InterchangeImport.UMaterialExpressionPlus", NewName="/Script/InterchangeImport.UMaterialExpressionMaterialXPlus")
-+ClassRedirects=(OldName="/Script/InterchangeImport.UMaterialExpressionPremult", NewName="/Script/InterchangeImport.UMaterialExpressionMaterialXPremult")
-+ClassRedirects=(OldName="/Script/InterchangeImport.UMaterialExpressionRamp4", NewName="/Script/InterchangeImport.UMaterialExpressionMaterialXRamp4")
-+ClassRedirects=(OldName="/Script/InterchangeImport.UMaterialExpressionRampLeftRight", NewName="/Script/InterchangeImport.UMaterialExpressionMaterialXRampLeftRight")
-+ClassRedirects=(OldName="/Script/InterchangeImport.UMaterialExpressionRampTopBottom", NewName="/Script/InterchangeImport.UMaterialExpressionMaterialXRampTopBottom")
-+ClassRedirects=(OldName="/Script/InterchangeImport.UMaterialExpressionRemap", NewName="/Script/InterchangeImport.UMaterialExpressionMaterialXRemap")
-+ClassRedirects=(OldName="/Script/InterchangeImport.UMaterialExpressionRgbToHsv", NewName="/Script/InterchangeImport.UMaterialExpressionMaterialXRgbToHsv")
-+ClassRedirects=(OldName="/Script/InterchangeImport.UMaterialExpressionRotate2D", NewName="/Script/InterchangeImport.UMaterialExpressionMaterialXRotate2D")
-+ClassRedirects=(OldName="/Script/InterchangeImport.UMaterialExpressionScreen", NewName="/Script/InterchangeImport.UMaterialExpressionMaterialXScreen")
-+ClassRedirects=(OldName="/Script/InterchangeImport.UMaterialExpressionSplitLeftRight", NewName="/Script/InterchangeImport.UMaterialExpressionMaterialXSplitLeftRight")
-+ClassRedirects=(OldName="/Script/InterchangeImport.UMaterialExpressionSplitTopBottom", NewName="/Script/InterchangeImport.UMaterialExpressionMaterialXSplitTopBottom")
-+ClassRedirects=(OldName="/Script/InterchangeImport.UMaterialExpressionTextureSampleParameterBlur", NewName="/Script/InterchangeImport.UMaterialExpressionMaterialXTextureSampleParameterBlur")
-+ClassRedirects=(OldName="/Script/InterchangeImport.UMaterialExpressionUnpremult", NewName="/Script/InterchangeImport.UMaterialExpressionMaterialXUnpremult")
++ClassRedirects=(OldName="/Script/InterchangeImport.MaterialExpressionAppend3Vector", NewName="/Script/InterchangeImport.MaterialExpressionMaterialXAppend3Vector")
++ClassRedirects=(OldName="/Script/InterchangeImport.MaterialExpressionAppend4Vector", NewName="/Script/InterchangeImport.MaterialExpressionMaterialXAppend4Vector")
++ClassRedirects=(OldName="/Script/InterchangeImport.MaterialExpressionBurn", NewName="/Script/InterchangeImport.MaterialExpressionMaterialXBurn")
++ClassRedirects=(OldName="/Script/InterchangeImport.MaterialExpressionDifference", NewName="/Script/InterchangeImport.MaterialExpressionMaterialXDifference")
++ClassRedirects=(OldName="/Script/InterchangeImport.MaterialExpressionDisjointOver", NewName="/Script/InterchangeImport.MaterialExpressionMaterialXDisjointOver")
++ClassRedirects=(OldName="/Script/InterchangeImport.MaterialExpressionDodge", NewName="/Script/InterchangeImport.MaterialExpressionMaterialXDodge")
++ClassRedirects=(OldName="/Script/InterchangeImport.MaterialExpressionExponential", NewName="/Script/InterchangeImport.MaterialExpressionMaterialXExponential")
++ClassRedirects=(OldName="/Script/InterchangeImport.MaterialExpressionFractal3D", NewName="/Script/InterchangeImport.MaterialExpressionMaterialXFractal3D")
++ClassRedirects=(OldName="/Script/InterchangeImport.MaterialExpressionHsvToRgb", NewName="/Script/InterchangeImport.MaterialExpressionMaterialXHsvToRgb")
++ClassRedirects=(OldName="/Script/InterchangeImport.MaterialExpressionIn", NewName="/Script/InterchangeImport.MaterialExpressionMaterialXIn")
++ClassRedirects=(OldName="/Script/InterchangeImport.MaterialExpressionLength", NewName="/Script/InterchangeImport.MaterialExpressionMaterialXLength")
++ClassRedirects=(OldName="/Script/InterchangeImport.MaterialExpressionLogarithm", NewName="/Script/InterchangeImport.MaterialExpressionMaterialXLogarithm")
++ClassRedirects=(OldName="/Script/InterchangeImport.MaterialExpressionLuminance", NewName="/Script/InterchangeImport.MaterialExpressionMaterialXLuminance")
++ClassRedirects=(OldName="/Script/InterchangeImport.MaterialExpressionMask", NewName="/Script/InterchangeImport.MaterialExpressionMaterialXMask")
++ClassRedirects=(OldName="/Script/InterchangeImport.MaterialExpressionMatte", NewName="/Script/InterchangeImport.MaterialExpressionMaterialXMatte")
++ClassRedirects=(OldName="/Script/InterchangeImport.MaterialExpressionMinus", NewName="/Script/InterchangeImport.MaterialExpressionMaterialXMinus")
++ClassRedirects=(OldName="/Script/InterchangeImport.MaterialExpressionOut", NewName="/Script/InterchangeImport.MaterialExpressionMaterialXOut")
++ClassRedirects=(OldName="/Script/InterchangeImport.MaterialExpressionOver", NewName="/Script/InterchangeImport.MaterialExpressionMaterialXOver")
++ClassRedirects=(OldName="/Script/InterchangeImport.MaterialExpressionOverlay", NewName="/Script/InterchangeImport.MaterialExpressionMaterialXOverlay")
++ClassRedirects=(OldName="/Script/InterchangeImport.MaterialExpressionPlace2D", NewName="/Script/InterchangeImport.MaterialExpressionMaterialXPlace2D")
++ClassRedirects=(OldName="/Script/InterchangeImport.MaterialExpressionPlus", NewName="/Script/InterchangeImport.MaterialExpressionMaterialXPlus")
++ClassRedirects=(OldName="/Script/InterchangeImport.MaterialExpressionPremult", NewName="/Script/InterchangeImport.MaterialExpressionMaterialXPremult")
++ClassRedirects=(OldName="/Script/InterchangeImport.MaterialExpressionRamp4", NewName="/Script/InterchangeImport.MaterialExpressionMaterialXRamp4")
++ClassRedirects=(OldName="/Script/InterchangeImport.MaterialExpressionRampLeftRight", NewName="/Script/InterchangeImport.MaterialExpressionMaterialXRampLeftRight")
++ClassRedirects=(OldName="/Script/InterchangeImport.MaterialExpressionRampTopBottom", NewName="/Script/InterchangeImport.MaterialExpressionMaterialXRampTopBottom")
++ClassRedirects=(OldName="/Script/InterchangeImport.MaterialExpressionRemap", NewName="/Script/InterchangeImport.MaterialExpressionMaterialXRemap")
++ClassRedirects=(OldName="/Script/InterchangeImport.MaterialExpressionRgbToHsv", NewName="/Script/InterchangeImport.MaterialExpressionMaterialXRgbToHsv")
++ClassRedirects=(OldName="/Script/InterchangeImport.MaterialExpressionRotate2D", NewName="/Script/InterchangeImport.MaterialExpressionMaterialXRotate2D")
++ClassRedirects=(OldName="/Script/InterchangeImport.MaterialExpressionScreen", NewName="/Script/InterchangeImport.MaterialExpressionMaterialXScreen")
++ClassRedirects=(OldName="/Script/InterchangeImport.MaterialExpressionSplitLeftRight", NewName="/Script/InterchangeImport.MaterialExpressionMaterialXSplitLeftRight")
++ClassRedirects=(OldName="/Script/InterchangeImport.MaterialExpressionSplitTopBottom", NewName="/Script/InterchangeImport.MaterialExpressionMaterialXSplitTopBottom")
++ClassRedirects=(OldName="/Script/InterchangeImport.MaterialExpressionTextureSampleParameterBlur", NewName="/Script/InterchangeImport.MaterialExpressionMaterialXTextureSampleParameterBlur")
++ClassRedirects=(OldName="/Script/InterchangeImport.MaterialExpressionUnpremult", NewName="/Script/InterchangeImport.MaterialExpressionMaterialXUnpremult")
 +EnumRedirects=(OldName="ETextureSampleBlurFilter", NewName="EMaterialXTextureSampleBlurFilter")
 +EnumRedirects=(OldName="ELuminanceMode", NewName="EMaterialXLuminanceMode")
 
