Enable the new VA error handling by default.

diff --git a/Engine/Config/BaseEngine.ini b/Engine/Config/BaseEngine.ini
--- a/Engine/Config/BaseEngine.ini
+++ b/Engine/Config/BaseEngine.ini
@@ -1399,7 +1399,7 @@ AllowSubmitIfVirtualizationFailed=false
 FilterMode=OptOut
 FilterMapContent=True
 LazyInitConnections=False
-UseLegacyErrorHandling=True
+UseLegacyErrorHandling=False
 ForceCachingOnPull=False
 PullErrorAdditionalMsg=""
 ; The following assets should not virtualize until it is confirmed that they do not crash when virtualized payloads are missing
