Exclude any .projectstore zenstore marker files when staging to pak files.
The .projectstore file is created if you cook to zenstore.
We want it when we stage to loose files as it contains information on how to connect to the zenserver where the bulk of the data exists.

diff --git a/Engine/Config/BasePakFileRules.ini b/Engine/Config/BasePakFileRules.ini
--- a/Engine/Config/BasePakFileRules.ini
+++ b/Engine/Config/BasePakFileRules.ini
@@ -22,3 +22,7 @@ bOverrideChunkManifest=true
 +Files=".../Engine/Content/EngineMaterials/DefaultCalibrationGrayscale.*"
 +Files=".../Engine/Content/EngineMaterials/PPM_DefaultCalibrationColor.*"
 +Files=".../Engine/Content/EngineMaterials/PPM_DefaultCalibrationGrayscale.*"
+
+[ExcludeProjectStoreMarkerFile]
+bExcludeFromPaks=true
++Files=".../.projectstore"
