Remove use of obsolete PurgeTransient parameter on DDC configurations.

diff --git a/Engine/Config/BaseEngine.ini b/Engine/Config/BaseEngine.ini
--- a/Engine/Config/BaseEngine.ini
+++ b/Engine/Config/BaseEngine.ini
@@ -2402,7 +2402,6 @@ Root=(Type=Hierarchical, Inner=Pak, Inner=EnginePak, Inner=Local, Inner=Shared,
 ; Configure the local cache. This is a FileSystem cache with the following parameters:
 ;	ReadOnly: Can data be written to this layer. Default is false.
 ;	Touch: Update timestamps on every access instead of only after 1/4 of UnusedFileAge. Default is false.
-;	PurgeTransient: Delete transient data from the cache. Default is false.
 ;	Clean: Perform a cleanup of old files on launch. Default is false.
 ;	Flush: Clear the DDC and start over on launch. Default is false.
 ;	DeleteUnused: Whether to clean up old files on a background thread. Default is true.
@@ -2419,7 +2418,7 @@ Root=(Type=Hierarchical, Inner=Pak, Inner=EnginePak, Inner=Local, Inner=Shared,
 ;		ValueOnly: Look only for compressed data in the 5.0+ directory layout.
 ;		LegacyOnly: Look only for uncompressed data in the pre-5.0 directory layout.
 ;		ValueWithLegacyFallback: Look in the 5.0+ location but fall back to the pre-5.0 location on miss.
-Local=(Type=FileSystem, PurgeTransient=true, UnusedFileAge=34, PromptIfMissing=true, Path=%ENGINEDIR%DerivedDataCache, EnvPathOverride=UE-LocalDataCachePath, EditorOverrideSetting=LocalDerivedDataCache, CommandLineOverride=LocalDataCachePath)
+Local=(Type=FileSystem, UnusedFileAge=34, PromptIfMissing=true, Path=%ENGINEDIR%DerivedDataCache, EnvPathOverride=UE-LocalDataCachePath, EditorOverrideSetting=LocalDerivedDataCache, CommandLineOverride=LocalDataCachePath)
 ; Configure the shared cache that is accessed after local.
 Shared=(Type=FileSystem, UnusedFileAge=10, FoldersToClean=10, ConsiderSlowAt=70, Path=?EpicDDC, EnvPathOverride=UE-SharedDataCachePath, EditorOverrideSetting=SharedDerivedDataCache, CommandLineOverride=SharedDataCachePath)
 ; Configure the cloud cache that is accessed after shared.
@@ -2434,7 +2433,7 @@ EnginePak=(Type=ReadPak, Filename="%ENGINEDIR%DerivedDataCache/DDC.ddp")
 Root=(Type=KeyLength, Length=120, Inner=AsyncPut)
 AsyncPut=(Type=AsyncPut, Inner=Hierarchy)
 Hierarchy=(Type=Hierarchical, Inner=Pak, Inner=CompressedPak, Inner=EnginePak, Inner=EnterprisePak, Inner=Local, Inner=Shared, Inner=Cloud)
-Local=(Type=FileSystem, PurgeTransient=true, UnusedFileAge=34, Path="%ENGINEVERSIONAGNOSTICUSERDIR%DerivedDataCache", EnvPathOverride=UE-LocalDataCachePath, EditorOverrideSetting=LocalDerivedDataCache)
+Local=(Type=FileSystem, UnusedFileAge=34, Path="%ENGINEVERSIONAGNOSTICUSERDIR%DerivedDataCache", EnvPathOverride=UE-LocalDataCachePath, EditorOverrideSetting=LocalDerivedDataCache)
 Shared=(Type=FileSystem, UnusedFileAge=10, FoldersToClean=10, ConsiderSlowAt=70, Path=?EpicDDC, EnvPathOverride=UE-SharedDataCachePath, EditorOverrideSetting=SharedDerivedDataCache, CommandLineOverride=SharedDataCachePath)
 Cloud=(Type=Cloud, ServerID=Default, EnvHostOverride=UE-CloudDataCacheHost, CommandLineHostOverride=CloudDataCacheHost, CommandLineOAuthSecretOverride="CloudDataCacheOAuthSecret", OAuthAccessTokenEnvOverride="UE-CloudDataCacheAccessToken")
 Pak=(Type=ReadPak, Filename="%GAMEDIR%DerivedDataCache/DDC.ddp")
@@ -2449,7 +2448,7 @@ Zen=(Type=Zen, Namespace="ue.ddc")
 
 [ZenDDC]
 Root=(Type=Hierarchical, Inner=Pak, Inner=EnginePak, Inner=Local, Inner=ZenLocal, Inner=Shared, Inner=Cloud)
-Local=(Type=FileSystem, PurgeTransient=true, DeleteOnly=true, UnusedFileAge=34, PromptIfMissing=true, Path=%ENGINEDIR%DerivedDataCache, EnvPathOverride=UE-LocalDataCachePath, EditorOverrideSetting=LocalDerivedDataCache, CommandLineOverride=LocalDataCachePath)
+Local=(Type=FileSystem, DeleteOnly=true, UnusedFileAge=34, PromptIfMissing=true, Path=%ENGINEDIR%DerivedDataCache, EnvPathOverride=UE-LocalDataCachePath, EditorOverrideSetting=LocalDerivedDataCache, CommandLineOverride=LocalDataCachePath)
 ZenLocal=(Type=Zen, Namespace="ue.ddc")
 Shared=(Type=FileSystem, UnusedFileAge=10, FoldersToClean=10, ConsiderSlowAt=70, Path=?EpicDDC, EnvPathOverride=UE-SharedDataCachePath, EditorOverrideSetting=SharedDerivedDataCache, CommandLineOverride=SharedDataCachePath)
 Cloud=(Type=Cloud, ServerID=Default, EnvHostOverride=UE-CloudDataCacheHost, CommandLineHostOverride=CloudDataCacheHost, CommandLineOAuthSecretOverride="CloudDataCacheOAuthSecret", OAuthAccessTokenEnvOverride="UE-CloudDataCacheAccessToken")
@@ -2472,7 +2471,7 @@ ZenLocal=(Type=Zen, Namespace="ue.ddc", Sandbox=Cold)
 Root=(Type=KeyLength, Length=120, Inner=AsyncPut)
 AsyncPut=(Type=AsyncPut, Inner=Hierarchy)
 Hierarchy=(Type=Hierarchical, Inner=Pak, Inner=EnginePak, Inner=Local, Inner=Seattle)
-Local=(Type=FileSystem, PurgeTransient=true, UnusedFileAge=34, Path=%ENGINEDIR%DerivedDataCache, EnvPathOverride=UE-LocalDataCachePath, EditorOverrideSetting=LocalDerivedDataCache)
+Local=(Type=FileSystem, UnusedFileAge=34, Path=%ENGINEDIR%DerivedDataCache, EnvPathOverride=UE-LocalDataCachePath, EditorOverrideSetting=LocalDerivedDataCache)
 Seattle=(Type=FileSystem, UnusedFileAge=23, FoldersToClean=10, Path=?EpicSeaDDC, EnvPathOverride=UE-SharedDataCachePath_Seattle)
 Pak=(Type=ReadPak, Filename="%GAMEDIR%DerivedDataCache/DDC.ddp")
 EnginePak=(Type=ReadPak, Filename="%ENGINEDIR%DerivedDataCache/DDC.ddp")
@@ -2481,7 +2480,7 @@ EnginePak=(Type=ReadPak, Filename="%ENGINEDIR%DerivedDataCache/DDC.ddp")
 Root=(Type=KeyLength, Length=120, Inner=AsyncPut)
 AsyncPut=(Type=AsyncPut, Inner=Hierarchy)
 Hierarchy=(Type=Hierarchical, Inner=Pak, Inner=Local)
-Local=(Type=FileSystem, PurgeTransient=true, UnusedFileAge=34, Path=%ENGINEDIR%DerivedDataCache, EnvPathOverride=UE-LocalDataCachePath, EditorOverrideSetting=LocalDerivedDataCache, CommandLineOverride=LocalDataCachePath)
+Local=(Type=FileSystem, UnusedFileAge=34, Path=%ENGINEDIR%DerivedDataCache, EnvPathOverride=UE-LocalDataCachePath, EditorOverrideSetting=LocalDerivedDataCache, CommandLineOverride=LocalDataCachePath)
 Pak=(Type=ReadPak, Filename="%GAMEDIR%DerivedDataCache/DDC.ddp")
 
 ; Creates a completely new standalone cache so as to not destroy a valuable populated local cache
@@ -2489,20 +2488,20 @@ Pak=(Type=ReadPak, Filename="%GAMEDIR%DerivedDataCache/DDC.ddp")
 Root=(Type=KeyLength, Length=120, Inner=AsyncPut)
 AsyncPut=(Type=AsyncPut, Inner=Hierarchy)
 Hierarchy=(Type=Hierarchical, Inner=Local)
-Local=(Type=FileSystem, Flush=true, PurgeTransient=true, UnusedFileAge=1, Path="%GAMEDIR%ColdDerivedDataCache")
+Local=(Type=FileSystem, Flush=true, UnusedFileAge=1, Path="%GAMEDIR%ColdDerivedDataCache")
 
 ; Inherits the same standalone cache from [Cold], for incremental performance measurement purposes
 [Warm]
 Root=(Type=KeyLength, Length=120, Inner=AsyncPut)
 AsyncPut=(Type=AsyncPut, Inner=Hierarchy)
 Hierarchy=(Type=Hierarchical, Inner=Local)
-Local=(Type=FileSystem, Clean=true, PurgeTransient=true, UnusedFileAge=1, Path="%GAMEDIR%ColdDerivedDataCache")
+Local=(Type=FileSystem, Clean=true, UnusedFileAge=1, Path="%GAMEDIR%ColdDerivedDataCache")
 
 [CreatePak]
 Root=(Type=KeyLength, Length=120, Inner=AsyncPut)
 AsyncPut=(Type=AsyncPut, Inner=Hierarchy)
 Hierarchy=(Type=Hierarchical, Inner=PakWrite, Inner=PakRead, Inner=Local, Inner=Shared, Inner=Cloud)
-Local=(Type=FileSystem, PurgeTransient=true, UnusedFileAge=34, Path=%ENGINEDIR%DerivedDataCache, EnvPathOverride=UE-LocalDataCachePath, EditorOverrideSetting=LocalDerivedDataCache)
+Local=(Type=FileSystem, UnusedFileAge=34, Path=%ENGINEDIR%DerivedDataCache, EnvPathOverride=UE-LocalDataCachePath, EditorOverrideSetting=LocalDerivedDataCache)
 Shared=(Type=FileSystem, DeleteUnused=false, ConsiderSlowAt=70, Path=?EpicDDC, EnvPathOverride=UE-SharedDataCachePath, EditorOverrideSetting=SharedDerivedDataCache, CommandLineOverride=SharedDataCachePath)
 Cloud=(Type=Cloud, ServerID=Default, EnvHostOverride=UE-CloudDataCacheHost, CommandLineHostOverride=CloudDataCacheHost, CommandLineOAuthSecretOverride="CloudDataCacheOAuthSecret", OAuthAccessTokenEnvOverride="UE-CloudDataCacheAccessToken")
 PakRead=(Type=ReadPak, Filename="%GAMEDIR%DerivedDataCache/DDC.ddp")
@@ -2512,7 +2511,7 @@ PakWrite=(Type=WritePak, Filename="%GAMEDIR%DerivedDataCache/DDC.ddp")
 Root=(Type=KeyLength, Length=120, Inner=AsyncPut)
 AsyncPut=(Type=AsyncPut, Inner=Hierarchy)
 Hierarchy=(Type=Hierarchical, Inner=EnginePak, Inner=PakWrite, Inner=PakRead, Inner=Local, Inner=Shared, Inner=Cloud)
-Local=(Type=FileSystem, PurgeTransient=true, UnusedFileAge=34, Path=%ENGINEDIR%DerivedDataCache, EnvPathOverride=UE-LocalDataCachePath, EditorOverrideSetting=LocalDerivedDataCache)
+Local=(Type=FileSystem, UnusedFileAge=34, Path=%ENGINEDIR%DerivedDataCache, EnvPathOverride=UE-LocalDataCachePath, EditorOverrideSetting=LocalDerivedDataCache)
 Shared=(Type=FileSystem, DeleteUnused=false, ConsiderSlowAt=70, Path=?EpicDDC, EnvPathOverride=UE-SharedDataCachePath, EditorOverrideSetting=SharedDerivedDataCache, CommandLineOverride=SharedDataCachePath)
 Cloud=(Type=Cloud, ServerID=Default, EnvHostOverride=UE-CloudDataCacheHost, CommandLineHostOverride=CloudDataCacheHost, CommandLineOAuthSecretOverride="CloudDataCacheOAuthSecret", OAuthAccessTokenEnvOverride="UE-CloudDataCacheAccessToken")
 EnginePak=(Type=ReadPak, Filename=../../../Engine/DerivedDataCache/Compressed.ddp, Compressed=true)
@@ -2523,7 +2522,7 @@ PakWrite=(Type=WritePak, Filename="%GAMEDIR%DerivedDataCache/Compressed.ddp", Co
 Root=(Type=KeyLength, Length=120, Inner=AsyncPut)
 AsyncPut=(Type=AsyncPut, Inner=Hierarchy)
 Hierarchy=(Type=Hierarchical, Inner=PakWrite, Inner=Local, Inner=Shared, Inner=Cloud)
-Local=(Type=FileSystem, PurgeTransient=true, UnusedFileAge=34, Path=%ENGINEDIR%DerivedDataCache, EnvPathOverride=UE-LocalDataCachePath, EditorOverrideSetting=LocalDerivedDataCache)
+Local=(Type=FileSystem, UnusedFileAge=34, Path=%ENGINEDIR%DerivedDataCache, EnvPathOverride=UE-LocalDataCachePath, EditorOverrideSetting=LocalDerivedDataCache)
 Shared=(Type=FileSystem, DeleteUnused=false, ConsiderSlowAt=70, Path=?EpicDDC, EnvPathOverride=UE-SharedDataCachePath, EditorOverrideSetting=SharedDerivedDataCache, CommandLineOverride=SharedDataCachePath)
 Cloud=(Type=Cloud, ServerID=Default, EnvHostOverride=UE-CloudDataCacheHost, CommandLineHostOverride=CloudDataCacheHost, CommandLineOAuthSecretOverride="CloudDataCacheOAuthSecret", OAuthAccessTokenEnvOverride="UE-CloudDataCacheAccessToken")
 PakWrite=(Type=WritePak, Filename="%ENGINEDIR%DerivedDataCache/Compressed.ddp", Compressed=true)
@@ -2532,7 +2531,7 @@ PakWrite=(Type=WritePak, Filename="%ENGINEDIR%DerivedDataCache/Compressed.ddp",
 Root=(Type=KeyLength, Length=120, Inner=AsyncPut)
 AsyncPut=(Type=AsyncPut, Inner=Hierarchy)
 Hierarchy=(Type=Hierarchical, Inner=PakWrite, Inner=Local, Inner=Shared, Inner=Cloud)
-Local=(Type=FileSystem, PurgeTransient=true, UnusedFileAge=34, Path=%ENGINEDIR%DerivedDataCache, EnvPathOverride=UE-LocalDataCachePath, EditorOverrideSetting=LocalDerivedDataCache)
+Local=(Type=FileSystem, UnusedFileAge=34, Path=%ENGINEDIR%DerivedDataCache, EnvPathOverride=UE-LocalDataCachePath, EditorOverrideSetting=LocalDerivedDataCache)
 Shared=(Type=FileSystem, DeleteUnused=false, ConsiderSlowAt=70, Path=?EpicDDC, EnvPathOverride=UE-SharedDataCachePath, EditorOverrideSetting=SharedDerivedDataCache, CommandLineOverride=SharedDataCachePath)
 Cloud=(Type=Cloud, ServerID=Default, EnvHostOverride=UE-CloudDataCacheHost, CommandLineHostOverride=CloudDataCacheHost, CommandLineOAuthSecretOverride="CloudDataCacheOAuthSecret", OAuthAccessTokenEnvOverride="UE-CloudDataCacheAccessToken")
 PakWrite=(Type=WritePak, Filename=../../../Enterprise/DerivedDataCache/Compressed.ddp, Compressed=true)
@@ -2541,8 +2540,8 @@ PakWrite=(Type=WritePak, Filename=../../../Enterprise/DerivedDataCache/Compresse
 Root=(Type=KeyLength, Length=120, Inner=AsyncPut)
 AsyncPut=(Type=AsyncPut, Inner=Hierarchy)
 Hierarchy=(Type=Hierarchical, Inner=Local, Inner=Project, Inner=Shared, Inner=Cloud)
-Local=(Type=FileSystem, PurgeTransient=true, UnusedFileAge=34, Path=%ENGINEDIR%DerivedDataCache, EnvPathOverride=UE-LocalDataCachePath, EditorOverrideSetting=LocalDerivedDataCache)
-Project=(Type=FileSystem, Flush=true, PurgeTransient=true, UnusedFileAge=34, Path=%GAMEDIR%ProjectDerivedData)
+Local=(Type=FileSystem, UnusedFileAge=34, Path=%ENGINEDIR%DerivedDataCache, EnvPathOverride=UE-LocalDataCachePath, EditorOverrideSetting=LocalDerivedDataCache)
+Project=(Type=FileSystem, Flush=true, UnusedFileAge=34, Path=%GAMEDIR%ProjectDerivedData)
 Shared=(Type=FileSystem, DeleteUnused=false, ConsiderSlowAt=70, Path=?EpicDDC, EnvPathOverride=UE-SharedDataCachePath, EditorOverrideSetting=SharedDerivedDataCache, CommandLineOverride=SharedDataCachePath)
 Cloud=(Type=Cloud, ServerID=Default, EnvHostOverride=UE-CloudDataCacheHost, CommandLineHostOverride=CloudDataCacheHost, CommandLineOAuthSecretOverride="CloudDataCacheOAuthSecret", OAuthAccessTokenEnvOverride="UE-CloudDataCacheAccessToken")
 
