Config to enable RT on Android SM5 devices

diff --git a/Engine/Config/Android/BaseAndroidEngine.ini b/Engine/Config/Android/BaseAndroidEngine.ini
--- a/Engine/Config/Android/BaseAndroidEngine.ini
+++ b/Engine/Config/Android/BaseAndroidEngine.ini
@@ -12,6 +12,9 @@ AudioCaptureModuleName=AudioCaptureAndroid
 [OnlineServices.EOS.Auth]
 DefaultExternalCredentialTypeStr=Google
 
+[/Script/AndroidTargetPlatform.AndroidTargetSettings]
+bEnableRayTracing=true
+
 ;Uncomment to turn on AndroidBackgroundHttp through the AndroidFetchBackgroundDownload and AndroidBackgroundService modules
 ;[BackgroundHttp]
 ;PlatformModularFeatureName=AndroidFetchBackgroundDownload
\ No newline at end of file
diff --git a/Engine/Config/Android/DataDrivenPlatformInfo.ini b/Engine/Config/Android/DataDrivenPlatformInfo.ini
--- a/Engine/Config/Android/DataDrivenPlatformInfo.ini
+++ b/Engine/Config/Android/DataDrivenPlatformInfo.ini
@@ -51,7 +51,7 @@ ActiveIconName=LevelEditor.PreviewMode.AndroidVulkanSM5.Enabled
 InactiveIconName=LevelEditor.PreviewMode.AndroidVulkanSM5.Disabled
 MenuTooltip=NSLOCTEXT("PreviewPlatform", "PreviewMenuTooltip_AndroidVulkanSM5", "Mobile preview using Android Vulkan SM5 quality settings.")
 IconText=NSLOCTEXT("PreviewPlatform", "PreviewIconText_AndroidVulkanSM5", "Android VK SM5")
-DeviceProfileName=Android
+DeviceProfileName=Android_Vulkan_SM5
 
 [ShaderPlatform OPENGL_ES3_1_ANDROID]
 Language=OpenGL
@@ -86,7 +86,19 @@ ShaderFormat=SF_VULKAN_SM5_ANDROID
 bTargetsTiledGPU=true
 bNeedsOfflineCompiler=true
 bSupportsIndexBufferUAVs=true
-bSupportsManualVertexFetch=false
 bRequiresDisableForwardLocalLights=true
 bWaterUsesSimpleForwardShading=true
-FriendlyName=LOCTEXT("FriendlyShaderPlatformName_AndroidVulkanSM5", "Android Vulkan SM5")
+; Raytracing and Lumen
+bSupportsManualVertexFetch=true
+bSupportsRayTracing = true
+bSupportsRayTracingShaders = false
+bSupportsInlineRayTracing = true
+bSupportsRayTracingIndirectInstanceData = true
+bSupportsByteBufferComputeShaders = true
+bEnableRayTracing=true
+bSupportsGPUSkinCache=true
+EnablesHLSL2021ByDefault=1
+bSupportsGPUScene=true
+bSupportsLumenGI=true
+bSupportsDistanceFields=true
+bSupportsDualSourceBlending=false
diff --git a/Engine/Config/BaseDeviceProfiles.ini b/Engine/Config/BaseDeviceProfiles.ini
--- a/Engine/Config/BaseDeviceProfiles.ini
+++ b/Engine/Config/BaseDeviceProfiles.ini
@@ -649,7 +649,9 @@ BaseProfileName=AppleTV4K
 +MatchProfile=(Profile="Android_Adreno5xx",Match=((SourceType=SRC_GpuFamily,CompareType=CMP_Regex,MatchString="Adreno \\(TM\\) 5[0-9][0-9]")))
 ; Example profiles for Vulkan SM5
 ;+MatchProfile=(Profile="Android_Vulkan_SM5",Match=((SourceType=SRC_GpuFamily,CompareType=CMP_Regex,MatchString="Adreno \\(TM\\) 6[0-9][0-9]"),(SourceType=SRC_AndroidVersion, CompareType=CMP_Regex,MatchString="([0-9]+).*"),(SourceType=SRC_PreviousRegexMatch,CompareType=CMP_GreaterEqual,MatchString="9"),(SourceType=SRC_VulkanAvailable,CompareType=CMP_Equal,MatchString="true")))
++MatchProfile=(Profile="Android_Vulkan_SM5",Match=((SourceType=SRC_GpuFamily,CompareType=CMP_Regex,MatchString="Adreno \\(TM\\) 7[0-9][0-9]"),(SourceType=SRC_AndroidVersion, CompareType=CMP_Regex,MatchString="([0-9]+).*"),(SourceType=SRC_PreviousRegexMatch,CompareType=CMP_GreaterEqual,MatchString="9"),(SourceType=SRC_VulkanAvailable,CompareType=CMP_Equal,MatchString="true")))
 ;+MatchProfile=(Profile="Android_Vulkan_SM5_Mali",Match=((SourceType=SRC_GpuFamily,CompareType=CMP_Regex,MatchString="^Mali\\-G7[67]"),(SourceType=SRC_VulkanAvailable,CompareType=CMP_Equal,MatchString="true")))
++MatchProfile=(Profile="Android_Vulkan_SM5_Mali",Match=((SourceType=SRC_GpuFamily,CompareType=CMP_Regex,MatchString="^Mali\\-G715"),(SourceType=SRC_VulkanAvailable,CompareType=CMP_Equal,MatchString="true")))
 ; disable Vulkan on Adreno 660 devices with driver version <= 525
 +MatchProfile=(Profile="Android_Adreno6xx",Match=((SourceType=SRC_GpuFamily,CompareType=CMP_Regex,MatchString="Adreno \\(TM\\) 660"),(SourceType=SRC_GLVersion, CompareType=CMP_Regex,MatchString="V@([0-9]+)"),(SourceType=SRC_PreviousRegexMatch,CompareType=CMP_LessEqual,MatchString="525")))
 +MatchProfile=(Profile="Android_Adreno6xx_Vulkan",Match=((SourceType=SRC_GpuFamily,CompareType=CMP_Regex,MatchString="Adreno \\(TM\\) 6[0-9][0-9]"),(SourceType=SRC_AndroidVersion, CompareType=CMP_Regex,MatchString="([0-9]+).*"),(SourceType=SRC_PreviousRegexMatch,CompareType=CMP_GreaterEqual,MatchString="10"),(SourceType=SRC_VulkanAvailable,CompareType=CMP_Equal,MatchString="true")))
@@ -680,6 +682,8 @@ BaseProfileName=AppleTV4K
 +MatchProfile=(Profile="Android_Unknown_Vulkan",Match=((SourceType=SRC_GpuFamily,CompareType=CMP_Regex,MatchString="^Mali"),(SourceType=SRC_AndroidVersion, CompareType=CMP_Regex,MatchString="([0-9]+).*"),(SourceType=SRC_PreviousRegexMatch,CompareType=CMP_GreaterEqual,MatchString="10"),(SourceType=SRC_VulkanAvailable,CompareType=CMP_Equal,MatchString="true")))
 
 ; Samsung Exynos / Xclipse GPU-based devices use Vulkan by default.
++MatchProfile=(Profile="Android_Vulkan_SM5",Match=((SourceType=SRC_GpuFamily,CompareType=CMP_Regex,MatchString="Samsung Xclipse 920"),(SourceType=SRC_VulkanAvailable,CompareType=CMP_Equal,MatchString="true")))
++MatchProfile=(Profile="Android_Vulkan_SM5_Mali",Match=((SourceType=SRC_GpuFamily,CompareType=CMP_Regex,MatchString="^Mali\\-G7[67]"),(SourceType=SRC_VulkanAvailable,CompareType=CMP_Equal,MatchString="true")))
 +MatchProfile=(Profile="Android_Xclipse_920_Vulkan",Match=((SourceType=SRC_GpuFamily,CompareType=CMP_Regex,MatchString="Samsung Xclipse 920"),(SourceType=SRC_VulkanAvailable,CompareType=CMP_Equal,MatchString="true")))
 +MatchProfile=(Profile="Android_Xclipse_920",Match=((SourceType=SRC_GpuFamily,CompareType=CMP_Regex,MatchString="Samsung Xclipse 920")))
 
@@ -711,6 +715,7 @@ bIsVisibleForAssets=True
 +CVars=r.Android.DisableVulkanSM5Support=1
 ; PF_B8G8R8A8
 +CVars=r.DefaultBackBufferPixelFormat=0 
++CVars=r.DistanceFields=0
 ; PreviewAllowlistCVars and PreviewDenyListCVars are arrays of cvars (or prefixes) that are included or excluded from being applied in editor preview.
 ; If any PreviewAllowlistCVars is set, cvars are denied by default.
 PreviewAllowlistCVars=sg.
@@ -801,6 +806,14 @@ BaseProfileName=Android
 +CVars=r.VirtualTextures=1
 +CVars=r.Android.DisableVulkanSM5Support=0
 +CVars=r.Android.DisableVulkanSupport=0
++CVars=r.DistanceFields=1
++CVars=r.Vulkan.RayTracing.AllowCompaction=0
++CVars=r.Vulkan.RayTracing.TLASPreferFastTraceTLAS=0
+; Disable rough diffuse for saving lighting cost
++CVars=r.Substrate.RoughDiffuse=0
+; Reduce shading quality for performance
++CVars=r.Substrate.ShadingQuality=2
++CVars=r.Substrate.TileCoord8bits=1
 
 [Android_Vulkan_SM5_Mali DeviceProfile]
 DeviceType=Android
diff --git a/Engine/Config/BaseEngine.ini b/Engine/Config/BaseEngine.ini
--- a/Engine/Config/BaseEngine.ini
+++ b/Engine/Config/BaseEngine.ini
@@ -2352,7 +2352,7 @@ HighDetail=(NumLODs=6,MaxNumStreamedLODs=0,bSupportLODStreaming=0,LODPercentTria
 PathToRAC=RuntimeAssetCache
 
 [Zen]
-AutoLaunch=true
+AutoLaunch=false
 
 [Zen.AutoLaunch]
 DataPath=%APPSETTINGSDIR%Zen/Data
@@ -2946,6 +2946,7 @@ bDisableFunctionDataSections=False
 bEnableAdvancedBinaryCompression=True
 bDisableStackProtector=False
 bEnableMulticastSupport=False
+bEnableRayTracing=true
 
 [/Script/AndroidPlatformEditor.AndroidSDKSettings]
 SDKAPILevel=latest
diff --git a/Engine/Source/Runtime/VulkanRHI/Private/Android/VulkanAndroidPlatform.h b/Engine/Source/Runtime/VulkanRHI/Private/Android/VulkanAndroidPlatform.h
--- a/Engine/Source/Runtime/VulkanRHI/Private/Android/VulkanAndroidPlatform.h
+++ b/Engine/Source/Runtime/VulkanRHI/Private/Android/VulkanAndroidPlatform.h
@@ -34,11 +34,7 @@
 //	VK_IMAGE_LAYOUT_DEPTH_ATTACHMENT_STENCIL_READ_ONLY_OPTIMAL = 1000117001,
 #define VULKAN_USE_REAL_RENDERPASS_COMPATIBILITY	0
 
-#if VULKAN_RHI_RAYTRACING
-#define UE_VK_API_VERSION							VK_API_VERSION_1_2
-#else
 #define UE_VK_API_VERSION							VK_API_VERSION_1_1
-#endif // VULKAN_RHI_RAYTRACING
 
 #define ENUM_VK_ENTRYPOINTS_PLATFORM_BASE(EnumMacro)
 
