Restore autolaunch of ZenServer.

diff --git a/Engine/Config/BaseEngine.ini b/Engine/Config/BaseEngine.ini
--- a/Engine/Config/BaseEngine.ini
+++ b/Engine/Config/BaseEngine.ini
@@ -2352,7 +2352,7 @@ HighDetail=(NumLODs=6,MaxNumStreamedLODs=0,bSupportLODStreaming=0,LODPercentTria
 PathToRAC=RuntimeAssetCache
 
 [Zen]
-AutoLaunch=false
+AutoLaunch=true
 
 [Zen.AutoLaunch]
 DataPath=%APPSETTINGSDIR%Zen/Data
