Disabled SM5 Android platforms by default

diff --git a/Engine/Config/BaseDeviceProfiles.ini b/Engine/Config/BaseDeviceProfiles.ini
--- a/Engine/Config/BaseDeviceProfiles.ini
+++ b/Engine/Config/BaseDeviceProfiles.ini
@@ -649,9 +649,9 @@ BaseProfileName=AppleTV4K
 +MatchProfile=(Profile="Android_Adreno5xx",Match=((SourceType=SRC_GpuFamily,CompareType=CMP_Regex,MatchString="Adreno \\(TM\\) 5[0-9][0-9]")))
 ; Example profiles for Vulkan SM5
 ;+MatchProfile=(Profile="Android_Vulkan_SM5",Match=((SourceType=SRC_GpuFamily,CompareType=CMP_Regex,MatchString="Adreno \\(TM\\) 6[0-9][0-9]"),(SourceType=SRC_AndroidVersion, CompareType=CMP_Regex,MatchString="([0-9]+).*"),(SourceType=SRC_PreviousRegexMatch,CompareType=CMP_GreaterEqual,MatchString="9"),(SourceType=SRC_VulkanAvailable,CompareType=CMP_Equal,MatchString="true")))
-+MatchProfile=(Profile="Android_Vulkan_SM5",Match=((SourceType=SRC_GpuFamily,CompareType=CMP_Regex,MatchString="Adreno \\(TM\\) 7[0-9][0-9]"),(SourceType=SRC_AndroidVersion, CompareType=CMP_Regex,MatchString="([0-9]+).*"),(SourceType=SRC_PreviousRegexMatch,CompareType=CMP_GreaterEqual,MatchString="9"),(SourceType=SRC_VulkanAvailable,CompareType=CMP_Equal,MatchString="true")))
+;+MatchProfile=(Profile="Android_Vulkan_SM5",Match=((SourceType=SRC_GpuFamily,CompareType=CMP_Regex,MatchString="Adreno \\(TM\\) 7[0-9][0-9]"),(SourceType=SRC_AndroidVersion, CompareType=CMP_Regex,MatchString="([0-9]+).*"),(SourceType=SRC_PreviousRegexMatch,CompareType=CMP_GreaterEqual,MatchString="9"),(SourceType=SRC_VulkanAvailable,CompareType=CMP_Equal,MatchString="true")))
 ;+MatchProfile=(Profile="Android_Vulkan_SM5_Mali",Match=((SourceType=SRC_GpuFamily,CompareType=CMP_Regex,MatchString="^Mali\\-G7[67]"),(SourceType=SRC_VulkanAvailable,CompareType=CMP_Equal,MatchString="true")))
-+MatchProfile=(Profile="Android_Vulkan_SM5_Mali",Match=((SourceType=SRC_GpuFamily,CompareType=CMP_Regex,MatchString="^Mali\\-G715"),(SourceType=SRC_VulkanAvailable,CompareType=CMP_Equal,MatchString="true")))
+;+MatchProfile=(Profile="Android_Vulkan_SM5_Mali",Match=((SourceType=SRC_GpuFamily,CompareType=CMP_Regex,MatchString="^Mali\\-G715"),(SourceType=SRC_VulkanAvailable,CompareType=CMP_Equal,MatchString="true")))
 ; disable Vulkan on Adreno 660 devices with driver version <= 525
 +MatchProfile=(Profile="Android_Adreno6xx",Match=((SourceType=SRC_GpuFamily,CompareType=CMP_Regex,MatchString="Adreno \\(TM\\) 660"),(SourceType=SRC_GLVersion, CompareType=CMP_Regex,MatchString="V@([0-9]+)"),(SourceType=SRC_PreviousRegexMatch,CompareType=CMP_LessEqual,MatchString="525")))
 +MatchProfile=(Profile="Android_Adreno6xx_Vulkan",Match=((SourceType=SRC_GpuFamily,CompareType=CMP_Regex,MatchString="Adreno \\(TM\\) 6[0-9][0-9]"),(SourceType=SRC_AndroidVersion, CompareType=CMP_Regex,MatchString="([0-9]+).*"),(SourceType=SRC_PreviousRegexMatch,CompareType=CMP_GreaterEqual,MatchString="10"),(SourceType=SRC_VulkanAvailable,CompareType=CMP_Equal,MatchString="true")))
@@ -682,8 +682,8 @@ BaseProfileName=AppleTV4K
 +MatchProfile=(Profile="Android_Unknown_Vulkan",Match=((SourceType=SRC_GpuFamily,CompareType=CMP_Regex,MatchString="^Mali"),(SourceType=SRC_AndroidVersion, CompareType=CMP_Regex,MatchString="([0-9]+).*"),(SourceType=SRC_PreviousRegexMatch,CompareType=CMP_GreaterEqual,MatchString="10"),(SourceType=SRC_VulkanAvailable,CompareType=CMP_Equal,MatchString="true")))
 
 ; Samsung Exynos / Xclipse GPU-based devices use Vulkan by default.
-+MatchProfile=(Profile="Android_Vulkan_SM5",Match=((SourceType=SRC_GpuFamily,CompareType=CMP_Regex,MatchString="Samsung Xclipse 920"),(SourceType=SRC_VulkanAvailable,CompareType=CMP_Equal,MatchString="true")))
-+MatchProfile=(Profile="Android_Vulkan_SM5_Mali",Match=((SourceType=SRC_GpuFamily,CompareType=CMP_Regex,MatchString="^Mali\\-G7[67]"),(SourceType=SRC_VulkanAvailable,CompareType=CMP_Equal,MatchString="true")))
+;+MatchProfile=(Profile="Android_Vulkan_SM5",Match=((SourceType=SRC_GpuFamily,CompareType=CMP_Regex,MatchString="Samsung Xclipse 920"),(SourceType=SRC_VulkanAvailable,CompareType=CMP_Equal,MatchString="true")))
+;+MatchProfile=(Profile="Android_Vulkan_SM5_Mali",Match=((SourceType=SRC_GpuFamily,CompareType=CMP_Regex,MatchString="^Mali\\-G7[67]"),(SourceType=SRC_VulkanAvailable,CompareType=CMP_Equal,MatchString="true")))
 +MatchProfile=(Profile="Android_Xclipse_920_Vulkan",Match=((SourceType=SRC_GpuFamily,CompareType=CMP_Regex,MatchString="Samsung Xclipse 920"),(SourceType=SRC_VulkanAvailable,CompareType=CMP_Equal,MatchString="true")))
 +MatchProfile=(Profile="Android_Xclipse_920",Match=((SourceType=SRC_GpuFamily,CompareType=CMP_Regex,MatchString="Samsung Xclipse 920")))
 
