Added ability to override the http version used for the http cache using "UE-CloudDataCacheHttpVersion" env variable or "-CloudDataCacheHttpVersion" command line argument

diff --git a/Engine/Config/BaseEngine.ini b/Engine/Config/BaseEngine.ini
--- a/Engine/Config/BaseEngine.ini
+++ b/Engine/Config/BaseEngine.ini
@@ -2413,6 +2413,9 @@ Root=(Type=Hierarchical, Inner=Pak, Inner=EnginePak, Inner=Local, Inner=Shared,
 ;	Path: Path to use for the filesystem DDC.
 ;	EnvPathOverride: An environment variable that if set will be used instead of Path. E.g., UE-LocalDataCachePath=D:\DDC. ('None' disables the DDC)
 ;	CommandLineOverride: A command line argument used in preference to the Path / EnvPathOverride settings. E.g., -SharedDataCachePath=\\someshare\folder
+;   HttpVersion: An override for which http version to use
+;   EnvHttpVersionOverride: An environment variable that if set will be used instead of HttpVersion.
+;   CommandLineHttpVersionOverride: A command line argument used in preference to the HttpVersion / EnvHttpVersionOverride settings
 ;	EditorOverrideSetting: Editor user setting that overrides the Path / EnvPathOverride line values.
 ;	LegacyMode: Where to look for legacy cache keys in storage. One of ValueOnly, ValueWithLegacyFallback, LegacyOnly.
 ;		ValueOnly: Look only for compressed data in the 5.0+ directory layout.
@@ -2422,7 +2425,7 @@ Local=(Type=FileSystem, UnusedFileAge=34, PromptIfMissing=true, Path=%ENGINEDIR%
 ; Configure the shared cache that is accessed after local.
 Shared=(Type=FileSystem, UnusedFileAge=10, FoldersToClean=10, ConsiderSlowAt=70, Path=?EpicDDC, EnvPathOverride=UE-SharedDataCachePath, EditorOverrideSetting=SharedDerivedDataCache, CommandLineOverride=SharedDataCachePath)
 ; Configure the cloud cache that is accessed after shared.
-Cloud=(Type=Cloud, ServerID=Default, EnvHostOverride=UE-CloudDataCacheHost, CommandLineHostOverride=CloudDataCacheHost, CommandLineOAuthSecretOverride="CloudDataCacheOAuthSecret", OAuthAccessTokenEnvOverride="UE-CloudDataCacheAccessToken")
+Cloud=(Type=Cloud, ServerID=Default, EnvHostOverride=UE-CloudDataCacheHost, CommandLineHostOverride=CloudDataCacheHost, CommandLineOAuthSecretOverride="CloudDataCacheOAuthSecret", OAuthAccessTokenEnvOverride="UE-CloudDataCacheAccessToken", EnvHttpVersionOverride="UE-CloudDataCacheHttpVersion", CommandLineHttpVersionOverride="CloudDataCacheHttpVersion")
 ; Configure a Project Pak node. This is a pre-generated DDC data file for the project that can be distributed to reduce runtime fetches/generation
 ; See documentation for how to create a DDP via the DerivedDataCache commandlet
 Pak=(Type=ReadPak, Filename="%GAMEDIR%DerivedDataCache/DDC.ddp")
@@ -2435,7 +2438,7 @@ AsyncPut=(Type=AsyncPut, Inner=Hierarchy)
 Hierarchy=(Type=Hierarchical, Inner=Pak, Inner=CompressedPak, Inner=EnginePak, Inner=EnterprisePak, Inner=Local, Inner=Shared, Inner=Cloud)
 Local=(Type=FileSystem, UnusedFileAge=34, Path="%ENGINEVERSIONAGNOSTICUSERDIR%DerivedDataCache", EnvPathOverride=UE-LocalDataCachePath, EditorOverrideSetting=LocalDerivedDataCache)
 Shared=(Type=FileSystem, UnusedFileAge=10, FoldersToClean=10, ConsiderSlowAt=70, Path=?EpicDDC, EnvPathOverride=UE-SharedDataCachePath, EditorOverrideSetting=SharedDerivedDataCache, CommandLineOverride=SharedDataCachePath)
-Cloud=(Type=Cloud, ServerID=Default, EnvHostOverride=UE-CloudDataCacheHost, CommandLineHostOverride=CloudDataCacheHost, CommandLineOAuthSecretOverride="CloudDataCacheOAuthSecret", OAuthAccessTokenEnvOverride="UE-CloudDataCacheAccessToken")
+Cloud=(Type=Cloud, ServerID=Default, EnvHostOverride=UE-CloudDataCacheHost, CommandLineHostOverride=CloudDataCacheHost, CommandLineOAuthSecretOverride="CloudDataCacheOAuthSecret", OAuthAccessTokenEnvOverride="UE-CloudDataCacheAccessToken", EnvHttpVersionOverride="UE-CloudDataCacheHttpVersion", CommandLineHttpVersionOverride="CloudDataCacheHttpVersion")
 Pak=(Type=ReadPak, Filename="%GAMEDIR%DerivedDataCache/DDC.ddp")
 CompressedPak=(Type=ReadPak, Filename="%GAMEDIR%DerivedDataCache/Compressed.ddp", Compressed=true)
 EnginePak=(Type=ReadPak, Filename=../../../Engine/DerivedDataCache/Compressed.ddp, Compressed=true)
@@ -2451,7 +2454,7 @@ Root=(Type=Hierarchical, Inner=Pak, Inner=EnginePak, Inner=Local, Inner=ZenLocal
 Local=(Type=FileSystem, DeleteOnly=true, UnusedFileAge=34, PromptIfMissing=true, Path=%ENGINEDIR%DerivedDataCache, EnvPathOverride=UE-LocalDataCachePath, EditorOverrideSetting=LocalDerivedDataCache, CommandLineOverride=LocalDataCachePath)
 ZenLocal=(Type=Zen, Namespace="ue.ddc")
 Shared=(Type=FileSystem, UnusedFileAge=10, FoldersToClean=10, ConsiderSlowAt=70, Path=?EpicDDC, EnvPathOverride=UE-SharedDataCachePath, EditorOverrideSetting=SharedDerivedDataCache, CommandLineOverride=SharedDataCachePath)
-Cloud=(Type=Cloud, ServerID=Default, EnvHostOverride=UE-CloudDataCacheHost, CommandLineHostOverride=CloudDataCacheHost, CommandLineOAuthSecretOverride="CloudDataCacheOAuthSecret", OAuthAccessTokenEnvOverride="UE-CloudDataCacheAccessToken")
+Cloud=(Type=Cloud, ServerID=Default, EnvHostOverride=UE-CloudDataCacheHost, CommandLineHostOverride=CloudDataCacheHost, CommandLineOAuthSecretOverride="CloudDataCacheOAuthSecret", OAuthAccessTokenEnvOverride="UE-CloudDataCacheAccessToken", EnvHttpVersionOverride="UE-CloudDataCacheHttpVersion", CommandLineHttpVersionOverride="CloudDataCacheHttpVersion")
 Pak=(Type=ReadPak, Filename="%GAMEDIR%DerivedDataCache/DDC.ddp")
 EnginePak=(Type=ReadPak, Filename="%ENGINEDIR%DerivedDataCache/DDC.ddp")
 
@@ -2503,7 +2506,7 @@ AsyncPut=(Type=AsyncPut, Inner=Hierarchy)
 Hierarchy=(Type=Hierarchical, Inner=PakWrite, Inner=PakRead, Inner=Local, Inner=Shared, Inner=Cloud)
 Local=(Type=FileSystem, UnusedFileAge=34, Path=%ENGINEDIR%DerivedDataCache, EnvPathOverride=UE-LocalDataCachePath, EditorOverrideSetting=LocalDerivedDataCache)
 Shared=(Type=FileSystem, DeleteUnused=false, ConsiderSlowAt=70, Path=?EpicDDC, EnvPathOverride=UE-SharedDataCachePath, EditorOverrideSetting=SharedDerivedDataCache, CommandLineOverride=SharedDataCachePath)
-Cloud=(Type=Cloud, ServerID=Default, EnvHostOverride=UE-CloudDataCacheHost, CommandLineHostOverride=CloudDataCacheHost, CommandLineOAuthSecretOverride="CloudDataCacheOAuthSecret", OAuthAccessTokenEnvOverride="UE-CloudDataCacheAccessToken")
+Cloud=(Type=Cloud, ServerID=Default, EnvHostOverride=UE-CloudDataCacheHost, CommandLineHostOverride=CloudDataCacheHost, CommandLineOAuthSecretOverride="CloudDataCacheOAuthSecret", OAuthAccessTokenEnvOverride="UE-CloudDataCacheAccessToken", EnvHttpVersionOverride="UE-CloudDataCacheHttpVersion", CommandLineHttpVersionOverride="CloudDataCacheHttpVersion")
 PakRead=(Type=ReadPak, Filename="%GAMEDIR%DerivedDataCache/DDC.ddp")
 PakWrite=(Type=WritePak, Filename="%GAMEDIR%DerivedDataCache/DDC.ddp")
 
@@ -2513,7 +2516,7 @@ AsyncPut=(Type=AsyncPut, Inner=Hierarchy)
 Hierarchy=(Type=Hierarchical, Inner=EnginePak, Inner=PakWrite, Inner=PakRead, Inner=Local, Inner=Shared, Inner=Cloud)
 Local=(Type=FileSystem, UnusedFileAge=34, Path=%ENGINEDIR%DerivedDataCache, EnvPathOverride=UE-LocalDataCachePath, EditorOverrideSetting=LocalDerivedDataCache)
 Shared=(Type=FileSystem, DeleteUnused=false, ConsiderSlowAt=70, Path=?EpicDDC, EnvPathOverride=UE-SharedDataCachePath, EditorOverrideSetting=SharedDerivedDataCache, CommandLineOverride=SharedDataCachePath)
-Cloud=(Type=Cloud, ServerID=Default, EnvHostOverride=UE-CloudDataCacheHost, CommandLineHostOverride=CloudDataCacheHost, CommandLineOAuthSecretOverride="CloudDataCacheOAuthSecret", OAuthAccessTokenEnvOverride="UE-CloudDataCacheAccessToken")
+Cloud=(Type=Cloud, ServerID=Default, EnvHostOverride=UE-CloudDataCacheHost, CommandLineHostOverride=CloudDataCacheHost, CommandLineOAuthSecretOverride="CloudDataCacheOAuthSecret", OAuthAccessTokenEnvOverride="UE-CloudDataCacheAccessToken", EnvHttpVersionOverride="UE-CloudDataCacheHttpVersion", CommandLineHttpVersionOverride="CloudDataCacheHttpVersion")
 EnginePak=(Type=ReadPak, Filename=../../../Engine/DerivedDataCache/Compressed.ddp, Compressed=true)
 PakRead=(Type=ReadPak, Filename="%GAMEDIR%DerivedDataCache/Compressed.ddp", Compressed=true)
 PakWrite=(Type=WritePak, Filename="%GAMEDIR%DerivedDataCache/Compressed.ddp", Compressed=true)
@@ -2524,7 +2527,7 @@ AsyncPut=(Type=AsyncPut, Inner=Hierarchy)
 Hierarchy=(Type=Hierarchical, Inner=PakWrite, Inner=Local, Inner=Shared, Inner=Cloud)
 Local=(Type=FileSystem, UnusedFileAge=34, Path=%ENGINEDIR%DerivedDataCache, EnvPathOverride=UE-LocalDataCachePath, EditorOverrideSetting=LocalDerivedDataCache)
 Shared=(Type=FileSystem, DeleteUnused=false, ConsiderSlowAt=70, Path=?EpicDDC, EnvPathOverride=UE-SharedDataCachePath, EditorOverrideSetting=SharedDerivedDataCache, CommandLineOverride=SharedDataCachePath)
-Cloud=(Type=Cloud, ServerID=Default, EnvHostOverride=UE-CloudDataCacheHost, CommandLineHostOverride=CloudDataCacheHost, CommandLineOAuthSecretOverride="CloudDataCacheOAuthSecret", OAuthAccessTokenEnvOverride="UE-CloudDataCacheAccessToken")
+Cloud=(Type=Cloud, ServerID=Default, EnvHostOverride=UE-CloudDataCacheHost, CommandLineHostOverride=CloudDataCacheHost, CommandLineOAuthSecretOverride="CloudDataCacheOAuthSecret", OAuthAccessTokenEnvOverride="UE-CloudDataCacheAccessToken", EnvHttpVersionOverride="UE-CloudDataCacheHttpVersion", CommandLineHttpVersionOverride="CloudDataCacheHttpVersion")
 PakWrite=(Type=WritePak, Filename="%ENGINEDIR%DerivedDataCache/Compressed.ddp", Compressed=true)
 
 [CreateInstalledEnterprisePak]
@@ -2533,7 +2536,7 @@ AsyncPut=(Type=AsyncPut, Inner=Hierarchy)
 Hierarchy=(Type=Hierarchical, Inner=PakWrite, Inner=Local, Inner=Shared, Inner=Cloud)
 Local=(Type=FileSystem, UnusedFileAge=34, Path=%ENGINEDIR%DerivedDataCache, EnvPathOverride=UE-LocalDataCachePath, EditorOverrideSetting=LocalDerivedDataCache)
 Shared=(Type=FileSystem, DeleteUnused=false, ConsiderSlowAt=70, Path=?EpicDDC, EnvPathOverride=UE-SharedDataCachePath, EditorOverrideSetting=SharedDerivedDataCache, CommandLineOverride=SharedDataCachePath)
-Cloud=(Type=Cloud, ServerID=Default, EnvHostOverride=UE-CloudDataCacheHost, CommandLineHostOverride=CloudDataCacheHost, CommandLineOAuthSecretOverride="CloudDataCacheOAuthSecret", OAuthAccessTokenEnvOverride="UE-CloudDataCacheAccessToken")
+Cloud=(Type=Cloud, ServerID=Default, EnvHostOverride=UE-CloudDataCacheHost, CommandLineHostOverride=CloudDataCacheHost, CommandLineOAuthSecretOverride="CloudDataCacheOAuthSecret", OAuthAccessTokenEnvOverride="UE-CloudDataCacheAccessToken", EnvHttpVersionOverride="UE-CloudDataCacheHttpVersion", CommandLineHttpVersionOverride="CloudDataCacheHttpVersion")
 PakWrite=(Type=WritePak, Filename=../../../Enterprise/DerivedDataCache/Compressed.ddp, Compressed=true)
 
 [CreateProjectCache]
diff --git a/Engine/Source/Developer/DerivedDataCache/Private/HttpCacheStore.cpp b/Engine/Source/Developer/DerivedDataCache/Private/HttpCacheStore.cpp
--- a/Engine/Source/Developer/DerivedDataCache/Private/HttpCacheStore.cpp
+++ b/Engine/Source/Developer/DerivedDataCache/Private/HttpCacheStore.cpp
@@ -219,6 +219,7 @@ struct FHttpCacheStoreParams
 	FString Host;
 	FString HostPinnedPublicKeys;
 	FString Namespace;
+	FString HttpVersion;
 	FString OAuthProvider;
 	FString OAuthClientId;
 	FString OAuthSecret;
@@ -308,6 +309,7 @@ private:
 	FString OAuthScope;
 	FString OAuthProviderIdentifier;
 	FString OAuthAccessToken;
+	FString HttpVersion;
 
 	FAnsiStringBuilderBase EffectiveDomain;
 
@@ -1558,6 +1560,7 @@ FHttpCacheStore::FHttpCacheStore(const FHttpCacheStoreParams& Params)
 	, OAuthScope(Params.OAuthScope)
 	, OAuthProviderIdentifier(Params.OAuthProviderIdentifier)
 	, OAuthAccessToken(Params.OAuthAccessToken)
+	, HttpVersion(Params.HttpVersion)
 	, bReadOnly(Params.bReadOnly)
 {
 	TRACE_CPUPROFILER_EVENT_SCOPE(HttpDDC_Construct);
@@ -1636,6 +1639,40 @@ FHttpCacheStore::~FHttpCacheStore()
 	}
 }
 
+template <typename CharType>
+static bool HttpVersionFromString(EHttpVersion& OutVersion, const TStringView<CharType> String)
+{
+	const auto ConvertedString = StringCast<UTF8CHAR, 16>(String.GetData(), String.Len());
+	if (ConvertedString == UTF8TEXTVIEW("none"))
+	{
+		OutVersion = EHttpVersion::None;
+	}
+	else if (ConvertedString == UTF8TEXTVIEW("http1.0"))
+	{
+		OutVersion = EHttpVersion::V1_0;
+	}
+	else if (ConvertedString == UTF8TEXTVIEW("http1.1"))
+	{
+		OutVersion = EHttpVersion::V1_1;
+	}
+	else if (ConvertedString == UTF8TEXTVIEW("http2"))
+	{
+		OutVersion = EHttpVersion::V2;
+	}
+	else if (ConvertedString == UTF8TEXTVIEW("http2-only"))
+	{
+		OutVersion = EHttpVersion::V2Only;
+	}
+	else
+	{
+		return false;
+	}
+	return true;
+}
+
+bool TryLexFromString(EHttpVersion& OutVersion, FUtf8StringView String) { return HttpVersionFromString(OutVersion, String); }
+bool TryLexFromString(EHttpVersion& OutVersion, FWideStringView String) { return HttpVersionFromString(OutVersion, String); }
+
 FHttpClientParams FHttpCacheStore::GetDefaultClientParams() const
 {
 	FHttpClientParams ClientParams;
@@ -1643,10 +1680,14 @@ FHttpClientParams FHttpCacheStore::GetDefaultClientParams() const
 	ClientParams.ConnectTimeout = 3 * 1000;
 	ClientParams.LowSpeedLimit = 1024;
 	ClientParams.LowSpeedTime = 10;
-	ClientParams.Version = EHttpVersion::V2;
 	ClientParams.TlsLevel = EHttpTlsLevel::All;
 	ClientParams.bFollowRedirects = true;
 	ClientParams.bFollow302Post = true;
+
+	EHttpVersion HttpVersionEnum = EHttpVersion::V2;
+	TryLexFromString(HttpVersionEnum, HttpVersion);
+	ClientParams.Version = HttpVersionEnum;
+
 	return ClientParams;
 }
 
@@ -2746,6 +2787,26 @@ void FHttpCacheStoreParams::Parse(const TCHAR* NodeName, const TCHAR* Config)
 
 	FParse::Bool(Config, TEXT("ResolveHostCanonicalName="), bResolveHostCanonicalName);
 
+	// Http version Params
+
+	FParse::Value(Config, TEXT("HttpVersion="), HttpVersion);
+	if (FParse::Value(Config, TEXT("EnvHttpVersionOverride="), OverrideName))
+	{
+		FString HttpEnv = FPlatformMisc::GetEnvironmentVariable(*OverrideName);
+		if (!HttpEnv.IsEmpty())
+		{
+			HttpVersion = HttpEnv;
+			UE_LOG(LogDerivedDataCache, Log, TEXT("%s: Found environment override for HttpVersion %s=%s"), NodeName, *OverrideName, *HttpVersion);
+		}
+	}
+	if (FParse::Value(Config, TEXT("CommandLineHttpVersionOverride="), OverrideName))
+	{
+		if (FParse::Value(FCommandLine::Get(), *(OverrideName + TEXT("=")), HttpVersion))
+		{
+			UE_LOG(LogDerivedDataCache, Log, TEXT("%s: Found command line override for HttpVersion %s=%s"), NodeName, *OverrideName, *HttpVersion);
+		}
+	}
+
 	// Namespace Params
 
 	if (Namespace.IsEmpty())
