Enable virtualization on static meshes and audio assets by default.

diff --git a/Engine/Config/BaseEngine.ini b/Engine/Config/BaseEngine.ini
--- a/Engine/Config/BaseEngine.ini
+++ b/Engine/Config/BaseEngine.ini
@@ -1405,9 +1405,6 @@ LazyInitConnections=False
 UseLegacyErrorHandling=False
 ForceCachingOnPull=False
 PullErrorAdditionalMsg=""
-; The following assets should not virtualize until it is confirmed that they do not crash when virtualized payloads are missing
-+DisabledAsset=StaticMesh
-+DisabledAsset=SoundWave
 
 [/Script/Virtualization.VirtualizationFilterSettings]
 !ExcludePackagePaths=ClearArray
