Data-driven ApplyAutoExposure for Buffer Visualization Modes
- specify whether a buffer visualization mode requires auto exposure in BaseEngine.ini instead of hardcoding in IsAutoExposureDebugMode(...)

#9897

diff --git a/Engine/Config/BaseEngine.ini b/Engine/Config/BaseEngine.ini
--- a/Engine/Config/BaseEngine.ini
+++ b/Engine/Config/BaseEngine.ini
@@ -2754,16 +2754,16 @@ DefaultUpdateOverlapsMethodDuringLevelStreaming=AlwaysUpdate
 BaseColor=(Material="/Engine/BufferVisualization/BaseColor.BaseColor", Name=LOCTEXT("BaseColorMat", "Base Color"))
 CustomDepth=(Material="/Engine/BufferVisualization/CustomDepth.CustomDepth", Name=LOCTEXT("BaseCustomDepthMat", "Custom Depth"))
 CustomStencil=(Material="/Engine/BufferVisualization/CustomStencil.CustomStencil", Name=LOCTEXT("BaseCustomStencilMat", "Custom Stencil"))
-FinalImage=(Material="/Engine/BufferVisualization/FinalImage.FinalImage", Name=LOCTEXT("BaseFinalImageMat", "Final Image"))
+FinalImage=(Material="/Engine/BufferVisualization/FinalImage.FinalImage", Name=LOCTEXT("BaseFinalImageMat", "Final Image"), ApplyAutoExposure=true)
 ShadingModel=(Material="/Engine/BufferVisualization/LightingModel.LightingModel", Name=LOCTEXT("BaseShadingModelMat", "Shading Model"))
 MaterialAO=(Material="/Engine/BufferVisualization/MaterialAO.MaterialAO", Name=LOCTEXT("BaseMaterialAOMat", "Material Ambient Occlusion"))
 Metallic=(Material="/Engine/BufferVisualization/Metallic.Metallic", Name=LOCTEXT("BaseMetallicMat", "Metallic"))
 Opacity=(Material="/Engine/BufferVisualization/Opacity.Opacity", Name=LOCTEXT("BaseOpacityMat", "Opacity"))
 Roughness=(Material="/Engine/BufferVisualization/Roughness.Roughness", Name=LOCTEXT("BaseRoughnessMat", "Roughness"))
 Anisotropy=(Material="/Engine/BufferVisualization/Anisotropy.Anisotropy", Name=LOCTEXT("BaseAnisotropyMat", "Anisotropy"))
-SceneColor=(Material="/Engine/BufferVisualization/SceneColor.SceneColor", Name=LOCTEXT("BaseSceneColorMat", "Scene Color"))
+SceneColor=(Material="/Engine/BufferVisualization/SceneColor.SceneColor", Name=LOCTEXT("BaseSceneColorMat", "Scene Color"), ApplyAutoExposure=true)
 SceneDepth=(Material="/Engine/BufferVisualization/SceneDepth.SceneDepth", Name=LOCTEXT("BaseSceneDepthMat", "Scene Depth"))
-SeparateTranslucencyRGB=(Material="/Engine/BufferVisualization/SeparateTranslucencyRGB.SeparateTranslucencyRGB", Name=LOCTEXT("BaseSeparateTranslucencyRGBMat", "Separate Translucency RGB"))
+SeparateTranslucencyRGB=(Material="/Engine/BufferVisualization/SeparateTranslucencyRGB.SeparateTranslucencyRGB", Name=LOCTEXT("BaseSeparateTranslucencyRGBMat", "Separate Translucency RGB"), ApplyAutoExposure=true)
 SeparateTranslucencyA=(Material="/Engine/BufferVisualization/SeparateTranslucencyA.SeparateTranslucencyA", Name=LOCTEXT("BaseSeparateTranslucencyAMat", "Separate Translucency Alpha"))
 Specular=(Material="/Engine/BufferVisualization/Specular.Specular", Name=LOCTEXT("BaseSpecularMat", "Specular"))
 SubsurfaceColor=(Material="/Engine/BufferVisualization/SubsurfaceColor.SubsurfaceColor", Name=LOCTEXT("BaseSubsurfaceColorMat", "Subsurface Color"))
@@ -2773,8 +2773,8 @@ AmbientOcclusion=(Material="/Engine/BufferVisualization/AmbientOcclusion.Ambient
 CustomDepthWorldUnits=(Material="/Engine/BufferVisualization/CustomDepthWorldUnits.CustomDepthWorldUnits", Name=LOCTEXT("BaseCustomDepthWorldUnitsMat", "Custom Depth World Units"))
 SceneDepthWorldUnits=(Material="/Engine/BufferVisualization/SceneDepthWorldUnits.SceneDepthWorldUnits", Name=LOCTEXT("BaseSceneDepthWorldUnitsMat", "Scene Depth World Units"))
 Velocity=(Material="/Engine/BufferVisualization/Velocity.Velocity", Name=LOCTEXT("Velocity", "Velocity"))
-PreTonemapHDRColor=(Material="/Engine/BufferVisualization/PreTonemapHDRColor.PreTonemapHDRColor", Name=LOCTEXT("PreTonemapHDRColor", "Pre Tonemap HDR Color"))
-PostTonemapHDRColor=(Material="/Engine/BufferVisualization/PostTonemapHDRColor.PostTonemapHDRColor", Name=LOCTEXT("PostTonemapHDRColor", "Post Tonemap HDR Color"))
+PreTonemapHDRColor=(Material="/Engine/BufferVisualization/PreTonemapHDRColor.PreTonemapHDRColor", Name=LOCTEXT("PreTonemapHDRColor", "Pre Tonemap HDR Color"), ApplyAutoExposure=true)
+PostTonemapHDRColor=(Material="/Engine/BufferVisualization/PostTonemapHDRColor.PostTonemapHDRColor", Name=LOCTEXT("PostTonemapHDRColor", "Post Tonemap HDR Color"), ApplyAutoExposure=true)
 
 [DeviceProfileManager]
 DeviceProfileSelectionModule="ExampleDeviceProfileSelector"
diff --git a/Engine/Source/Runtime/Engine/Private/BufferVisualizationData.cpp b/Engine/Source/Runtime/Engine/Private/BufferVisualizationData.cpp
--- a/Engine/Source/Runtime/Engine/Private/BufferVisualizationData.cpp
+++ b/Engine/Source/Runtime/Engine/Private/BufferVisualizationData.cpp
@@ -58,13 +58,19 @@ void FBufferVisualizationData::Initialize()
 			
 						if (Material)
 						{
+							FText DisplayName;
+							FParse::Value(*It.Value().GetValue(), TEXT("Name="), DisplayName, TEXT("Engine.BufferVisualizationMaterials"));
+
+							bool bApplyAutoExposure = false;
+							FParse::Bool(*It.Value().GetValue(), TEXT("ApplyAutoExposure="), bApplyAutoExposure);
+
 							Material->AddToRoot();
 							Record& Rec = MaterialMap.Add(It.Key(), Record());
 							Rec.Name = It.Key().GetPlainNameString();
 							Rec.Material = Material;
-							FText DisplayName;
-							FParse::Value( *It.Value().GetValue(), TEXT("Name="), DisplayName, TEXT("Engine.BufferVisualizationMaterials") );
+							Rec.bApplyAutoExposure = bApplyAutoExposure;
 							Rec.DisplayName = DisplayName;
+
 							MaterialMapFromMaterialName.Add(Material->GetFName(), Rec);
 						}
 					}
@@ -120,12 +126,12 @@ void FBufferVisualizationData::ConfigureConsoleCommand()
 		);
 }
 
-UMaterialInterface* FBufferVisualizationData::GetMaterial(FName InMaterialName)
+const FBufferVisualizationData::Record* FBufferVisualizationData::GetRecord(FName InMaterialName) const
 {
 	// Get UMaterial from the TMaterialMap FName
 	if (const Record* Result = MaterialMap.Find(InMaterialName))
 	{
-		return Result->Material;
+		return Result;
 	}
 	// Get UMaterial from the UMaterial FName
 	// Almost all BufferVisualizationData variables contain a FMaterial with its same FName.
@@ -133,7 +139,7 @@ UMaterialInterface* FBufferVisualizationData::GetMaterial(FName InMaterialName)
 	// This "else if" case handles the case in which we look for "LightingModel" rather than "ShadingModel", returning the same value
 	else if (const Record* ResultFromMaterialName = MaterialMapFromMaterialName.Find(InMaterialName))
 	{
-		return ResultFromMaterialName->Material;
+		return ResultFromMaterialName;
 	}
 	// Not found
 	else
@@ -142,22 +148,36 @@ UMaterialInterface* FBufferVisualizationData::GetMaterial(FName InMaterialName)
 	}
 }
 
-FText FBufferVisualizationData::GetMaterialDisplayName(FName InMaterialName) const
+UMaterialInterface* FBufferVisualizationData::GetMaterial(FName InMaterialName) const
 {
-	// Get UMaterial from the TMaterialMap FName
-	if (const Record* Result = MaterialMap.Find(InMaterialName))
+	if (const Record* Result = GetRecord(InMaterialName))
 	{
-		return Result->DisplayName;
+		return Result->Material;
 	}
-	// Get UMaterial from the UMaterial FName
-	// Almost all BufferVisualizationData variables contain a FMaterial with its same FName.
-	// But not all, e.g., ShadingModel uses the FMaterial LightingModel. This could confuse the developer depending on which one they are using
-	// This "else if" case handles the case in which we look for "LightingModel" rather than "ShadingModel", returning the same value
-	else if (const Record* ResultFromMaterialName = MaterialMapFromMaterialName.Find(InMaterialName))
+	else
 	{
-		return ResultFromMaterialName->DisplayName;
+		return nullptr;
+	}
+}
+
+bool FBufferVisualizationData::GetMaterialApplyAutoExposure(FName InMaterialName) const
+{
+	if (const Record* Result = GetRecord(InMaterialName))
+	{
+		return Result->bApplyAutoExposure;
+	}
+	else
+	{
+		return false;
+	}
+}
+
+FText FBufferVisualizationData::GetMaterialDisplayName(FName InMaterialName) const
+{
+	if (const Record* Result = GetRecord(InMaterialName))
+	{
+		return Result->DisplayName;
 	}
-	// Not found
 	else
 	{
 		return FText::GetEmpty();
diff --git a/Engine/Source/Runtime/Engine/Public/BufferVisualizationData.h b/Engine/Source/Runtime/Engine/Public/BufferVisualizationData.h
--- a/Engine/Source/Runtime/Engine/Public/BufferVisualizationData.h
+++ b/Engine/Source/Runtime/Engine/Public/BufferVisualizationData.h
@@ -21,7 +21,10 @@ public:
 	bool IsInitialized() const { return bIsInitialized; }
 
 	/** Get a named material from the available material map **/
-	ENGINE_API UMaterialInterface* GetMaterial(FName InMaterialName);
+	ENGINE_API UMaterialInterface* GetMaterial(FName InMaterialName) const;
+
+	/** Get whether auto exposure should be applied for the visualization mode **/
+	ENGINE_API bool GetMaterialApplyAutoExposure(FName InMaterialName) const;
 
 	/** Get the display name of a named material from the available material map **/
 	ENGINE_API FText GetMaterialDisplayName(FName InMaterialName) const;
@@ -60,6 +63,7 @@ private:
 		FString Name;
 		FText DisplayName;
 		UMaterialInterface* Material;
+		bool bApplyAutoExposure;
 	};
 
 	/** Mapping of FName (first parameter in ini file material list) to a material record */
@@ -86,6 +90,9 @@ private:
 
 	/** Internal helper function for creating the buffer visualization system console commands */
 	void ConfigureConsoleCommand();
+
+	/** Internal helper function to get record associated with material */
+	const Record* GetRecord(FName InMaterialName) const;
 };
 
 ENGINE_API FBufferVisualizationData& GetBufferVisualizationData();
diff --git a/Engine/Source/Runtime/Renderer/Private/PostProcess/PostProcessEyeAdaptation.cpp b/Engine/Source/Runtime/Renderer/Private/PostProcess/PostProcessEyeAdaptation.cpp
--- a/Engine/Source/Runtime/Renderer/Private/PostProcess/PostProcessEyeAdaptation.cpp
+++ b/Engine/Source/Runtime/Renderer/Private/PostProcess/PostProcessEyeAdaptation.cpp
@@ -6,6 +6,7 @@
 #include "TranslucentLighting.h"
 #include "Strata/Strata.h"
 
+#include "BufferVisualizationData.h"
 #include "SceneTextureParameters.h"
 #include "SystemTextures.h"
 
@@ -331,10 +332,7 @@ bool IsAutoExposureDebugMode(const FViewInfo& View)
 
 	return View.Family->UseDebugViewPS() ||
 		!EngineShowFlags.Lighting ||
-		(EngineShowFlags.VisualizeBuffer && View.CurrentBufferVisualizationMode != NAME_None
-			// Some HDR buffer visualization do want to have exposure applied.
-			&& View.CurrentBufferVisualizationMode != TEXT("SeparateTranslucencyRGB") && View.CurrentBufferVisualizationMode != TEXT("SceneColor") && View.CurrentBufferVisualizationMode != TEXT("FinalImage")
-			&& View.CurrentBufferVisualizationMode != TEXT("PreTonemapHDRColor") && View.CurrentBufferVisualizationMode != TEXT("PostTonemapHDRColor")) ||
+		(EngineShowFlags.VisualizeBuffer && View.CurrentBufferVisualizationMode != NAME_None && !GetBufferVisualizationData().GetMaterialApplyAutoExposure(View.CurrentBufferVisualizationMode)) ||
 		EngineShowFlags.RayTracingDebug ||
 		EngineShowFlags.VisualizeDistanceFieldAO ||
 		EngineShowFlags.VisualizeVolumetricCloudConservativeDensity ||
