Enable ACLPlugin by default and make it the default compression codec

diff --git a/Engine/Config/BaseEngine.ini b/Engine/Config/BaseEngine.ini
--- a/Engine/Config/BaseEngine.ini
+++ b/Engine/Config/BaseEngine.ini
@@ -3203,7 +3203,7 @@ bStripAnimationDataOnDedicatedServer=False
 BoneTimecodeCustomAttributeNameSettings=(HourAttributeName="TCHour",MinuteAttributeName="TCMinute",SecondAttributeName="TCSecond",FrameAttributeName="TCFrame",SubframeAttributeName="TCSubframe",RateAttributeName="TCRate",TakenameAttributeName="Takename")
 
 [Animation.DefaultObjectSettings]
-BoneCompressionSettings="/Engine/Animation/DefaultAnimBoneCompressionSettings"
+BoneCompressionSettings="/ACLPlugin/ACLAnimBoneCompressionSettings"
 AnimationRecorderBoneCompressionSettings="/Engine/Animation/DefaultRecorderBoneCompression"
 CurveCompressionSettings="/Engine/Animation/DefaultAnimCurveCompressionSettings"
 
diff --git a/Engine/Plugins/Animation/ACLPlugin/ACLPlugin.uplugin b/Engine/Plugins/Animation/ACLPlugin/ACLPlugin.uplugin
--- a/Engine/Plugins/Animation/ACLPlugin/ACLPlugin.uplugin
+++ b/Engine/Plugins/Animation/ACLPlugin/ACLPlugin.uplugin
@@ -10,7 +10,7 @@
 	"DocsURL": "",
 	"MarketplaceURL": "",
 	"SupportURL": "",
-	"EnabledByDefault": false,
+	"EnabledByDefault": true,
 	"CanContainContent": true,
 	"IsBetaVersion": false,
 	"Modules": [
