When launching ZenServer ensure that we set a minimum amount of disk space (2GiB) so that we don't end up filling the user's hard drive.

diff --git a/Engine/Config/BaseEngine.ini b/Engine/Config/BaseEngine.ini
--- a/Engine/Config/BaseEngine.ini
+++ b/Engine/Config/BaseEngine.ini
@@ -2356,7 +2356,7 @@ DataPath=%APPSETTINGSDIR%Zen/Data
 LocalDataCachePathEnvOverride=UE-LocalDataCachePath
 LocalDataCachePathCommandLineOverride=LocalDataCachePath
 LocalDataCachePathEditorOverrideSetting=LocalDerivedDataCache
-ExtraArgs=--http asio --gc-cache-duration-seconds 2937600 --gc-interval-seconds 21600
+ExtraArgs=--http asio --gc-cache-duration-seconds 2937600 --gc-interval-seconds 21600 --gc-low-diskspace-threshold 2147483648
 
 [Zen.ConnectExisting]
 HostName=localhost
