Reverse modern xcode submitted by mistake

diff --git a/Engine/Config/BaseEngine.ini b/Engine/Config/BaseEngine.ini
--- a/Engine/Config/BaseEngine.ini
+++ b/Engine/Config/BaseEngine.ini
@@ -3029,7 +3029,7 @@ bBuildAllSupportedOnBuildMachine=true
 bEnableRayTracing=false
 
 [/Script/MacTargetPlatform.XcodeProjectSettings]
-bUseModernXcode=true
+bUseModernXcode=false
 bUseModernCodeSigning=true
 ModernSigningTeam=
 ModernSigningPrefix=com.YourCompany
