Bumped nVidia minimum version to 527.27 : a random issue with landscape tools has been known to plague drivers <= 526.67. The exact cause is unknown but a fix has probably been made in one of the 527 releases : the first version to not exhibit the problem is 527.27 (cannot test minor releases, though)

diff --git a/Engine/Config/BaseHardware.ini b/Engine/Config/BaseHardware.ini
--- a/Engine/Config/BaseHardware.ini
+++ b/Engine/Config/BaseHardware.ini
@@ -36,12 +36,12 @@
 
 
 [GPU_NVIDIA Windows]
-SuggestedDriverVersion="526.47"
+SuggestedDriverVersion="527.27"
 
 ; To prevent problems with older drivers.
 +DriverDenyList=(DriverVersion="<473.47", Reason="These driver versions have known stability issues and missing features", RHI="D3D11")
-+DriverDenyList=(DriverVersion="<526.47", Reason="These driver versions have known stability issues and missing features", RHI="D3D12")
-+DriverDenyList=(DriverVersion="<526.47", Reason="These driver versions have known stability issues and missing features", RHI="Vulkan")
++DriverDenyList=(DriverVersion="<527.27", Reason="These driver versions have known stability issues and missing features", RHI="D3D12")
++DriverDenyList=(DriverVersion="<527.27", Reason="These driver versions have known stability issues and missing features", RHI="Vulkan")
 
 ;-----------------------------------------------------------------------------------------------------------------
 
