Fix: "Unable to find audio format BINKA from hinted modules, loading all potential format modules to find it"

diff --git a/Engine/Config/BaseEngine.ini b/Engine/Config/BaseEngine.ini
--- a/Engine/Config/BaseEngine.ini
+++ b/Engine/Config/BaseEngine.ini
@@ -1450,7 +1450,7 @@ DeferStartupPrecache=false
 +FormatModuleHints=AudioFormatOPUS
 +FormatModuleHints=AudioFormatOGG
 +FormatModuleHints=AudioFormatADPCM
-+FormatModuleHints=AudioFormatBINKA
++FormatModuleHints=AudioFormatBINK
 +AudioInfoModules=BinkAudioDecoder
 FallbackFormat=BINKA
 PlatformFormat=BINKA
