Remove hair scalability settings which are no longer used.

diff --git a/Engine/Config/BaseScalability.ini b/Engine/Config/BaseScalability.ini
--- a/Engine/Config/BaseScalability.ini
+++ b/Engine/Config/BaseScalability.ini
@@ -778,34 +778,29 @@ grass.DensityScale=1.0
 r.HairStrands.SkyLighting.IntegrationType=2
 r.HairStrands.SkyAO.SampleCount=4
 r.HairStrands.Visibility.MSAA.SamplePerPixel=1
-r.HairStrands.Interpolation.UseSingleGuide=1
 r.AnisotropicMaterials=0
 
 [ShadingQuality@1]
 r.HairStrands.SkyLighting.IntegrationType=2
 r.HairStrands.SkyAO.SampleCount=4
 r.HairStrands.Visibility.MSAA.SamplePerPixel=1
-r.HairStrands.Interpolation.UseSingleGuide=1
 r.AnisotropicMaterials=0
 
 [ShadingQuality@2]
 r.HairStrands.SkyLighting.IntegrationType=2
 r.HairStrands.SkyAO.SampleCount=4
 r.HairStrands.Visibility.MSAA.SamplePerPixel=4
-r.HairStrands.Interpolation.UseSingleGuide=1
 r.AnisotropicMaterials=1
 
 [ShadingQuality@3]
 r.HairStrands.SkyLighting.IntegrationType=2
 r.HairStrands.SkyAO.SampleCount=4
 r.HairStrands.Visibility.MSAA.SamplePerPixel=4
-r.HairStrands.Interpolation.UseSingleGuide=0
 r.AnisotropicMaterials=1
 
 [ShadingQuality@Cine]
 r.HairStrands.SkyLighting.IntegrationType=1
 r.HairStrands.SkyAO.SampleCount=8
 r.HairStrands.Visibility.MSAA.SamplePerPixel=8
-r.HairStrands.Interpolation.UseSingleGuide=0
 r.AnisotropicMaterials=1
 
