Reduce Autolaunched ZenServer garbage collection threshold from 34 days to 14 days.  This will reduce how long data is kept in the ZenServer local DDC as well as any cooked output stored in ZenServer.

diff --git a/Engine/Config/BaseEngine.ini b/Engine/Config/BaseEngine.ini
--- a/Engine/Config/BaseEngine.ini
+++ b/Engine/Config/BaseEngine.ini
@@ -2326,7 +2326,7 @@ DataPath=%APPSETTINGSDIR%Zen/Data
 LocalDataCachePathEnvOverride=UE-LocalDataCachePath
 LocalDataCachePathCommandLineOverride=LocalDataCachePath
 LocalDataCachePathEditorOverrideSetting=LocalDerivedDataCache
-ExtraArgs=--http asio --gc-cache-duration-seconds 2937600 --gc-interval-seconds 21600 --gc-low-diskspace-threshold 2147483648
+ExtraArgs=--http asio --gc-cache-duration-seconds 1209600 --gc-interval-seconds 21600 --gc-low-diskspace-threshold 2147483648
 
 [Zen.ConnectExisting]
 HostName=localhost
