Setting r.Forward.LightLinkedListCulling to 0 for platforms that do not benefit from it, saves about 0.1ms GPU

diff --git a/Engine/Config/BaseDeviceProfiles.ini b/Engine/Config/BaseDeviceProfiles.ini
--- a/Engine/Config/BaseDeviceProfiles.ini
+++ b/Engine/Config/BaseDeviceProfiles.ini
@@ -281,6 +281,9 @@ bIsVisibleForAssets=True
 +CVars=sg.TextureQuality=3
 +CVars=sg.EffectsQuality=3
 +CVars=sg.FoliageQuality=3
+; LightGrid Settings
++CVars=r.Forward.LightLinkedListCulling=0
+
 ; PreviewAllowlistCVars and PreviewDenyListCVars are arrays of cvars (or prefixes) that are included or excluded from being applied in editor preview
 ; If any PreviewAllowlistCVars is set, cvars are denied by default.
 PreviewAllowlistCVars=sg.
@@ -716,6 +719,8 @@ bIsVisibleForAssets=True
 ; PF_B8G8R8A8
 +CVars=r.DefaultBackBufferPixelFormat=0 
 +CVars=r.DistanceFields=0
+; LightGrid Settings
++CVars=r.Forward.LightLinkedListCulling=0
 ; PreviewAllowlistCVars and PreviewDenyListCVars are arrays of cvars (or prefixes) that are included or excluded from being applied in editor preview.
 ; If any PreviewAllowlistCVars is set, cvars are denied by default.
 PreviewAllowlistCVars=sg.
