- Removed empty struct ini lines, as empty strings will error when parsed as a struct

diff --git a/Engine/Config/BaseEngine.ini b/Engine/Config/BaseEngine.ini
--- a/Engine/Config/BaseEngine.ini
+++ b/Engine/Config/BaseEngine.ini
@@ -3025,8 +3025,6 @@ ShippingSpecificMacEntitlements=(FilePath="/Game/Build/Mac/Resources/Sandbox.NoN
 ; The following are used only with Manual codesigning
 MacSigningIdentity=Apple Development
 IOSSigningIdentity=Apple Development
-IOSProvisioningProfile=
-TVOSProvisioningProfile=
 
 [HMDPluginPriority]
 ; Since SteamVR also works with the Oculus Rift, give priority to the native Oculus plugin before trying SteamVR.
