Localization: Changed ShouldUseLocalizedNodeAndPinNames in BaseEditorSettings.ini to be false by default. Regional users find untranslated versions more useful.
#rnx

diff --git a/Engine/Config/BaseEditorSettings.ini b/Engine/Config/BaseEditorSettings.ini
--- a/Engine/Config/BaseEditorSettings.ini
+++ b/Engine/Config/BaseEditorSettings.ini
@@ -5,7 +5,7 @@
 [Internationalization]
 ShouldUseLocalizedNumericInput=True
 ShouldUseLocalizedPropertyNames=True
-ShouldUseLocalizedNodeAndPinNames=True
+ShouldUseLocalizedNodeAndPinNames=False
 
 [/Script/UnrealEd.EditorSettings]
 bCopyStarterContentPreference = true
