Run the GatherText asset scan async during editor boot so that it's likely finished before the commandlet needs the data

For a very basic loc gather this can take the asset gather time from ~120 seconds to ~40 seconds

diff --git a/Engine/Config/BaseEngine.ini b/Engine/Config/BaseEngine.ini
--- a/Engine/Config/BaseEngine.ini
+++ b/Engine/Config/BaseEngine.ini
@@ -2591,6 +2591,7 @@ bFilterSearchableNames=true
 ; Commandlets by default do not SearchAllAssets because many do not need it. A commandlet can be opted in
 ; using CommandletsUsingAR. Format is the shortname of the commandlet, optionally without the "Commandlet" suffix.
 +CommandletsUsingAR=DataValidation
++CommandletsUsingAR=GatherText
 
 [AutomationTesting]
 ImportTestPath=Content/EditorAutomation/
