- Hooked up the VisionPro simulator deviceprofile to disable virtual streaming
- Made Simulator deice strings for IOS platforms, iPhone0,0 - AppleTV0,0 - VisionOS0,0, since before they were just "arm"

diff --git a/Engine/Config/BaseDeviceProfiles.ini b/Engine/Config/BaseDeviceProfiles.ini
--- a/Engine/Config/BaseDeviceProfiles.ini
+++ b/Engine/Config/BaseDeviceProfiles.ini
@@ -68,7 +68,8 @@
 +DeviceProfileNameAndTypes=AppleTV4K,IOS
 +DeviceProfileNameAndTypes=AppleTV2_4K,IOS
 +DeviceProfileNameAndTypes=TVOS,TVOS
-+DeviceProfileNameAndTypes=VisionPro,IOS
++DeviceProfileNameAndTypes=VisionPro,VisionOS
++DeviceProfileNameAndTypes=VisionProSimulator,VisionOS
 +DeviceProfileNameAndTypes=VisionOS,VisionOS
 +DeviceProfileNameAndTypes=MacEditor,Mac
 +DeviceProfileNameAndTypes=MacClient,Mac
@@ -254,6 +255,8 @@ iPhone=iPhone14
 iPod=iPodTouch7
 iPad=iPad10
 AppleTV=AppleTV2_4K
+VisionPro0.*=VisionProSimulator
+VisionPro1.*=VisionPro
 
 [IOS DeviceProfile]
 DeviceType=IOS
@@ -640,9 +643,14 @@ DeviceType=IOS
 BaseProfileName=AppleTV4K
 
 [VisionPro DeviceProfile]
-DeviceType=IOS
+DeviceType=VisionOS
 BaseProfileName=IOS
 
+[VisionProSimulator DeviceProfile]
+DeviceType=VisionOS
+BaseProfileName=VisionPro
++CVars=r.Mobile.VirtualTextures=0
+
 [VisionOS DeviceProfile]
 DeviceType=VisionOS
 BaseProfileName=VisionPro
diff --git a/Engine/Plugins/Runtime/IOSDeviceProfileSelector/IOSDeviceProfileSelector.uplugin b/Engine/Plugins/Runtime/IOSDeviceProfileSelector/IOSDeviceProfileSelector.uplugin
--- a/Engine/Plugins/Runtime/IOSDeviceProfileSelector/IOSDeviceProfileSelector.uplugin
+++ b/Engine/Plugins/Runtime/IOSDeviceProfileSelector/IOSDeviceProfileSelector.uplugin
@@ -23,8 +23,9 @@
 			"PlatformAllowList" :
 			[
 				"IOS",
-				"TVOS"
+				"TVOS",
+				"VisionOS"
 			]
 		}
 	]
-}
\ No newline at end of file
+}
diff --git a/Engine/Source/Runtime/Core/Private/IOS/IOSPlatformMisc.cpp b/Engine/Source/Runtime/Core/Private/IOS/IOSPlatformMisc.cpp
--- a/Engine/Source/Runtime/Core/Private/IOS/IOSPlatformMisc.cpp
+++ b/Engine/Source/Runtime/Core/Private/IOS/IOSPlatformMisc.cpp
@@ -380,6 +380,19 @@ FString GetIOSDeviceIDString()
 		bCached = true;
 
 		free(DeviceID);
+		
+		// arm simulator
+		// @todo test intel simulator
+		if (CachedResult == "arm64")
+		{
+#if PLATFORM_VISIONOS
+			CachedResult = TEXT("VisionPro0,1");
+#elif PLATFORM_TVOS
+			CachedResult = TEXT("AppleTV0,1");
+#else
+			CachedResult = TEXT("iPhone0,1");
+#endif
+		}
 	}
 
 	return CachedResult;
