Add Magic Leap Profile to BaseDeviceProfiles.ini.

diff --git a/Engine/Config/BaseDeviceProfiles.ini b/Engine/Config/BaseDeviceProfiles.ini
--- a/Engine/Config/BaseDeviceProfiles.ini
+++ b/Engine/Config/BaseDeviceProfiles.ini
@@ -126,6 +126,7 @@
 +DeviceProfileNameAndTypes=Oculus_Quest,Android
 +DeviceProfileNameAndTypes=Oculus_Quest2,Android
 +DeviceProfileNameAndTypes=HoloLens,HoloLens
++DeviceProfileNameAndTypes=MagicLeap_Vulkan,Android
 
 [GlobalDefaults DeviceProfile]
 TextureLODGroups=(Group=TEXTUREGROUP_World,MinLODSize=1,MaxLODSize=16384,LODBias=0,MinMagFilter=aniso,MipFilter=point,MipGenSettings=TMGS_SimpleAverage)
@@ -663,6 +664,7 @@ BaseProfileName=VisionPro
 ; Match Quest devices directly
 +MatchProfile=(Profile="Oculus_Quest",Match=((SourceType=SRC_DeviceMake,CompareType=CMP_Equal,MatchString="Oculus"), (SourceType=SRC_HMDSystemName,CompareType=CMP_Equal,MatchString="Oculus Quest")))
 +MatchProfile=(Profile="Oculus_Quest2",Match=((SourceType=SRC_DeviceMake,CompareType=CMP_Equal,MatchString="Oculus"), (SourceType=SRC_HMDSystemName,CompareType=CMP_Equal,MatchString="Oculus Quest2")))
++MatchProfile=(Profile="MagicLeap_Vulkan",Match=((SourceType=SRC_DeviceMake,CompareType=CMP_Equal,MatchString="MagicLeap"), (SourceType=SRC_VulkanAvailable,CompareType=CMP_Equal,MatchString="true")))
 
 ; Adreno-based devices
 +MatchProfile=(Profile="Android_Adreno4xx",Match=((SourceType=SRC_GpuFamily,CompareType=CMP_Regex,MatchString="Adreno \\(TM\\) 4[0-9][0-9]")))
@@ -1101,3 +1103,8 @@ DeviceType=HoloLens
 BaseProfileName=Windows
 +CVars=r.D3D11.Depth24Bit=1
 +CVars=r.D3D12.Depth24Bit=1
+
+[MagicLeap_Vulkan DeviceProfile]
+DeviceType=Android
+BaseProfileName=Android_Mid
++CVars=xr.DisableOpenXROnAndroidWithoutOculus=0
