CADKernel is the default mesher for CAD import
Update ConsoleVariables.ini

diff --git a/Engine/Config/ConsoleVariables.ini b/Engine/Config/ConsoleVariables.ini
--- a/Engine/Config/ConsoleVariables.ini
+++ b/Engine/Config/ConsoleVariables.ini
@@ -122,11 +122,17 @@ fx.SkipVectorVMBackendOptimizations=1
 ;ds.CADTranslator.OverwriteCache=true
 
 ; Uncomment to enable/disable the timer that kill the worker if the import time is unusually long. With this time control, the load of the corrupted file 
-; is canceled but the rest of the scene is imported. Default is enable
-;ds.CADTranslator.EnableTimeControl=1
+; is canceled but the rest of the scene is imported.
+; Default value is true
+;ds.CADTranslator.EnableTimeControl=true
 
-; Uncomment to enable/disable CADKernel tessellation. If disable, CAD SDK tessellator is used.
+; Uncomment to enable/disable the memory control used by static meshes to avoid running out of memory.\n\
+; if enable, a warning is displayed allowing the user to interrupt the import and have only data already processed or to continue with the risk of running out of memory.
 ; Default value is true
+;ds.CADTranslator.EnableMemoryControl=true
+
+; Uncomment to enable/disable CADKernel tessellation. If disable, CAD SDK tessellator is used.
+; Default value is false
 ;ds.CADTranslator.DisableCADKernelTessellation=false
 
 ; Set Welding threshold for Heal/Sew stitching methods (in cm)
@@ -152,7 +158,7 @@ fx.SkipVectorVMBackendOptimizations=1
 
 ; Perform a welding of the mesh to try to stitch mesh cracks
 ; This welding is performed respecting the ds.CADTranslator.StitchingTolerance
-; Default value is true
+; Default value is false
 ;ds.CADTranslator.SewMeshIfNeeded=false
 
 ; Scale factor to change the unit of the DMU (Only applies to TechSoft import.)
diff --git a/Engine/Plugins/Enterprise/DatasmithCADImporter/Source/CADTools/Private/CADOption.cpp b/Engine/Plugins/Enterprise/DatasmithCADImporter/Source/CADTools/Private/CADOption.cpp
--- a/Engine/Plugins/Enterprise/DatasmithCADImporter/Source/CADTools/Private/CADOption.cpp
+++ b/Engine/Plugins/Enterprise/DatasmithCADImporter/Source/CADTools/Private/CADOption.cpp
@@ -16,11 +16,13 @@ Default is MaxImportThreads = 0\n\
    -if EnableCADCache is false, the scene is read all at once\n"),
 	ECVF_Default);
 
-bool FImportParameters::bGDisableCADKernelTessellation = true;
+bool FImportParameters::bGDisableCADKernelTessellation = false;
 FAutoConsoleVariableRef GCADTranslatorDisableCADKernelTessellation(
 	TEXT("ds.CADTranslator.DisableCADKernelTessellation"),
 	FImportParameters::bGDisableCADKernelTessellation,
-	TEXT("Disable to use CAD import library tessellator.\n"),
+	TEXT("\
+Disable to use CAD import library tessellator.\n\
+Default is enable\n"),
 	ECVF_Default);
 
 bool FImportParameters::bGEnableCADCache = true;
