Add profile for meta quest pro and update profiles for quest 2 and quest.

Also use Oculus_Quest device profile as fallback for any Quest device.

diff --git a/Engine/Config/BaseDeviceProfiles.ini b/Engine/Config/BaseDeviceProfiles.ini
--- a/Engine/Config/BaseDeviceProfiles.ini
+++ b/Engine/Config/BaseDeviceProfiles.ini
@@ -123,8 +123,9 @@
 +DeviceProfileNameAndTypes=Android_PowerVR_GM9xxx_Vulkan,Android
 +DeviceProfileNameAndTypes=Android_TegraK1,Android
 +DeviceProfileNameAndTypes=Android_Unknown_Vulkan,Android
-+DeviceProfileNameAndTypes=Oculus_Quest,Android
++DeviceProfileNameAndTypes=Meta_Quest_Pro,Android
 +DeviceProfileNameAndTypes=Oculus_Quest2,Android
++DeviceProfileNameAndTypes=Oculus_Quest,Android
 +DeviceProfileNameAndTypes=HoloLens,HoloLens
 +DeviceProfileNameAndTypes=MagicLeap_Vulkan,Android
 
@@ -662,8 +663,9 @@ BaseProfileName=VisionPro
 ; These rules are checked in turn and stop at the profile that matches all criteria
 [/Script/AndroidDeviceProfileSelector.AndroidDeviceProfileMatchingRules]
 ; Match Quest devices directly
-+MatchProfile=(Profile="Oculus_Quest",Match=((SourceType=SRC_DeviceMake,CompareType=CMP_Equal,MatchString="Oculus"), (SourceType=SRC_HMDSystemName,CompareType=CMP_Equal,MatchString="Oculus Quest")))
++MatchProfile=(Profile="Meta_Quest_Pro",Match=((SourceType=SRC_DeviceMake,CompareType=CMP_Equal,MatchString="Oculus"), (SourceType=SRC_HMDSystemName,CompareType=CMP_Equal,MatchString="Meta Quest Pro")))
 +MatchProfile=(Profile="Oculus_Quest2",Match=((SourceType=SRC_DeviceMake,CompareType=CMP_Equal,MatchString="Oculus"), (SourceType=SRC_HMDSystemName,CompareType=CMP_Equal,MatchString="Oculus Quest2")))
++MatchProfile=(Profile="Oculus_Quest",Match=((SourceType=SRC_DeviceMake,CompareType=CMP_Equal,MatchString="Oculus")))
 +MatchProfile=(Profile="MagicLeap_Vulkan",Match=((SourceType=SRC_DeviceMake,CompareType=CMP_Equal,MatchString="MagicLeap"), (SourceType=SRC_VulkanAvailable,CompareType=CMP_Equal,MatchString="true")))
 
 ; Adreno-based devices
@@ -1023,15 +1025,23 @@ BaseProfileName=Android_High
 ; enable Vulkan
 +CVars=r.Android.DisableVulkanSupport=0
 
-[Oculus_Quest DeviceProfile]
+[Meta_Quest_Pro DeviceProfile]
 DeviceType=Android
-BaseProfileName=Android_Mid
-+CVars=r.VRS.HMDDynamicFoveation=1
-+CVars=r.Android.DisableVulkanSupport=0
+BaseProfileName=Oculus_Quest2
 
 [Oculus_Quest2 DeviceProfile]
 DeviceType=Android
 BaseProfileName=Oculus_Quest
++CVars=r.Mobile.Oculus.ForceSymmetric=1
+
+[Oculus_Quest DeviceProfile]
+DeviceType=Android
+BaseProfileName=Android_Mid
++CVars=xr.VRS.DynamicFoveation=1
++CVars=r.Android.DisableVulkanSupport=0
+; disable running Niagara compute shaders on the GPU. This can cause hard to trace crashes.
++CVars=fx.NiagaraAllowGPUParticles=0
++CVars=FX.AllowGPUSorting=0
 
 [Mac DeviceProfile]
 DeviceType=Mac
