Enables TSR Ressurection on High, Epic & Cinematic anti-aliasing scalability

diff --git a/Engine/Config/BaseScalability.ini b/Engine/Config/BaseScalability.ini
--- a/Engine/Config/BaseScalability.ini
+++ b/Engine/Config/BaseScalability.ini
@@ -59,6 +59,7 @@ r.TSR.ShadingRejection.Flickering=0
 r.TSR.ShadingRejection.TileOverscan=2
 r.TSR.Velocity.Extrapolation=0
 r.TSR.RejectionAntiAliasingQuality=0
+r.TSR.Resurrection=0
 
 [AntiAliasingQuality@1]
 r.FXAA.Quality=1
@@ -70,6 +71,7 @@ r.TSR.ShadingRejection.Flickering=0
 r.TSR.ShadingRejection.TileOverscan=2
 r.TSR.Velocity.Extrapolation=0
 r.TSR.RejectionAntiAliasingQuality=1
+r.TSR.Resurrection=0
 
 [AntiAliasingQuality@2]
 r.FXAA.Quality=3
@@ -81,6 +83,7 @@ r.TSR.ShadingRejection.Flickering=1
 r.TSR.ShadingRejection.TileOverscan=3
 r.TSR.Velocity.Extrapolation=0
 r.TSR.RejectionAntiAliasingQuality=1
+r.TSR.Resurrection=1
 
 [AntiAliasingQuality@3]
 r.FXAA.Quality=4
@@ -92,6 +95,7 @@ r.TSR.ShadingRejection.Flickering=1
 r.TSR.ShadingRejection.TileOverscan=3
 r.TSR.Velocity.Extrapolation=1
 r.TSR.RejectionAntiAliasingQuality=2
+r.TSR.Resurrection=1
 
 [AntiAliasingQuality@Cine]
 r.FXAA.Quality=5
@@ -103,6 +107,7 @@ r.TSR.ShadingRejection.Flickering=1
 r.TSR.ShadingRejection.TileOverscan=3
 r.TSR.Velocity.Extrapolation=1
 r.TSR.RejectionAntiAliasingQuality=2
+r.TSR.Resurrection=1
 
 
 ;-----------------------------------------------------------------------------------------------------------------
