Disable Vulkan on Adreno6xx devices with drivers version earlier than 400

diff --git a/Engine/Config/BaseDeviceProfiles.ini b/Engine/Config/BaseDeviceProfiles.ini
--- a/Engine/Config/BaseDeviceProfiles.ini
+++ b/Engine/Config/BaseDeviceProfiles.ini
@@ -679,7 +679,8 @@ BaseProfileName=VisionPro
 ;+MatchProfile=(Profile="Android_Vulkan_SM5_Mali",Match=((SourceType=SRC_GpuFamily,CompareType=CMP_Regex,MatchString="^Mali\\-G715"),(SourceType=SRC_VulkanAvailable,CompareType=CMP_Equal,MatchString="true")))
 ; disable Vulkan on Adreno 660 devices with driver version <= 525
 +MatchProfile=(Profile="Android_Adreno6xx",Match=((SourceType=SRC_GpuFamily,CompareType=CMP_Regex,MatchString="Adreno \\(TM\\) 660"),(SourceType=SRC_GLVersion, CompareType=CMP_Regex,MatchString="V@([0-9]+)"),(SourceType=SRC_PreviousRegexMatch,CompareType=CMP_LessEqual,MatchString="525")))
-+MatchProfile=(Profile="Android_Adreno6xx_Vulkan",Match=((SourceType=SRC_GpuFamily,CompareType=CMP_Regex,MatchString="Adreno \\(TM\\) 6[0-9][0-9]"),(SourceType=SRC_AndroidVersion, CompareType=CMP_Regex,MatchString="([0-9]+).*"),(SourceType=SRC_PreviousRegexMatch,CompareType=CMP_GreaterEqual,MatchString="10"),(SourceType=SRC_VulkanAvailable,CompareType=CMP_Equal,MatchString="true")))
+; Enable Vulkan on Adreno 6xx devices with Android >= 10 and driver version >= 400
++MatchProfile=(Profile="Android_Adreno6xx_Vulkan",Match=((SourceType=SRC_GpuFamily,CompareType=CMP_Regex,MatchString="Adreno \\(TM\\) 6[0-9][0-9]"),(SourceType=SRC_AndroidVersion, CompareType=CMP_Regex,MatchString="([0-9]+).*"),(SourceType=SRC_PreviousRegexMatch,CompareType=CMP_GreaterEqual,MatchString="10"),(SourceType=SRC_VulkanAvailable,CompareType=CMP_Equal,MatchString="true"),(SourceType=SRC_GLVersion, CompareType=CMP_Regex,MatchString="V@([0-9]+)"),(SourceType=SRC_PreviousRegexMatch,CompareType=CMP_GreaterEqual,MatchString="400")))
 +MatchProfile=(Profile="Android_Adreno6xx",Match=((SourceType=SRC_GpuFamily,CompareType=CMP_Regex,MatchString="Adreno \\(TM\\) 6[0-9][0-9]")))
 +MatchProfile=(Profile="Android_Adreno7xx_Vulkan",Match=((SourceType=SRC_GpuFamily,CompareType=CMP_Regex,MatchString="Adreno \\(TM\\) 7[0-9][0-9]"),(SourceType=SRC_AndroidVersion, CompareType=CMP_Regex,MatchString="([0-9]+).*"),(SourceType=SRC_PreviousRegexMatch,CompareType=CMP_GreaterEqual,MatchString="10"),(SourceType=SRC_VulkanAvailable,CompareType=CMP_Equal,MatchString="true")))
 +MatchProfile=(Profile="Android_Adreno7xx",Match=((SourceType=SRC_GpuFamily,CompareType=CMP_Regex,MatchString="Adreno \\(TM\\) 7[0-9][0-9]")))
