Commits to r.TSR.ShadingRejection.TileOverscan=3 on all AA scalabilities

diff --git a/Engine/Config/BaseScalability.ini b/Engine/Config/BaseScalability.ini
--- a/Engine/Config/BaseScalability.ini
+++ b/Engine/Config/BaseScalability.ini
@@ -56,7 +56,6 @@ r.TSR.History.R11G11B10=1
 r.TSR.History.ScreenPercentage=100
 r.TSR.History.UpdateQuality=0
 r.TSR.ShadingRejection.Flickering=0
-r.TSR.ShadingRejection.TileOverscan=2
 r.TSR.Velocity.Extrapolation=0
 r.TSR.RejectionAntiAliasingQuality=0
 r.TSR.Resurrection=0
@@ -68,7 +67,6 @@ r.TSR.History.R11G11B10=1
 r.TSR.History.ScreenPercentage=100
 r.TSR.History.UpdateQuality=1
 r.TSR.ShadingRejection.Flickering=0
-r.TSR.ShadingRejection.TileOverscan=2
 r.TSR.Velocity.Extrapolation=0
 r.TSR.RejectionAntiAliasingQuality=1
 r.TSR.Resurrection=0
@@ -80,7 +78,6 @@ r.TSR.History.R11G11B10=1
 r.TSR.History.ScreenPercentage=100
 r.TSR.History.UpdateQuality=2
 r.TSR.ShadingRejection.Flickering=1
-r.TSR.ShadingRejection.TileOverscan=3
 r.TSR.Velocity.Extrapolation=0
 r.TSR.RejectionAntiAliasingQuality=1
 r.TSR.Resurrection=1
@@ -92,7 +89,6 @@ r.TSR.History.R11G11B10=1
 r.TSR.History.ScreenPercentage=200
 r.TSR.History.UpdateQuality=3
 r.TSR.ShadingRejection.Flickering=1
-r.TSR.ShadingRejection.TileOverscan=3
 r.TSR.Velocity.Extrapolation=1
 r.TSR.RejectionAntiAliasingQuality=2
 r.TSR.Resurrection=1
@@ -104,7 +100,6 @@ r.TSR.History.R11G11B10=0
 r.TSR.History.ScreenPercentage=200
 r.TSR.History.UpdateQuality=3
 r.TSR.ShadingRejection.Flickering=1
-r.TSR.ShadingRejection.TileOverscan=3
 r.TSR.Velocity.Extrapolation=1
 r.TSR.RejectionAntiAliasingQuality=2
 r.TSR.Resurrection=1
diff --git a/Engine/Source/Runtime/Renderer/Private/PostProcess/TemporalSuperResolution.cpp b/Engine/Source/Runtime/Renderer/Private/PostProcess/TemporalSuperResolution.cpp
--- a/Engine/Source/Runtime/Renderer/Private/PostProcess/TemporalSuperResolution.cpp
+++ b/Engine/Source/Runtime/Renderer/Private/PostProcess/TemporalSuperResolution.cpp
@@ -147,7 +147,7 @@ TAutoConsoleVariable<int32> CVarTSRShadingTileOverscan(
 	TEXT("However chaining many convlutions in this tiles means that some convolutions on the edge arround are becoming corrupted ")
 	TEXT("and therefor need to overlap the tile by couple of padding to hide it. This controls is controled by default in the ")
 	TEXT(" anti-aliasing scalability settings."),
-	ECVF_Scalability | ECVF_RenderThreadSafe);
+	ECVF_RenderThreadSafe);
 
 TAutoConsoleVariable<float> CVarTSRShadingExposureOffset(
 	TEXT("r.TSR.ShadingRejection.ExposureOffset"), 3.0,
@@ -1977,7 +1977,7 @@ FDefaultTemporalUpscaler::FOutputs AddTemporalSuperResolutionPasses(
 		PermutationVector = FTSRRejectShadingCS::RemapPermutation(PermutationVector);
 
 		const int32 GroupTileSize = 32;
-		const int32 TileOverscan = FMath::Clamp(CVarTSRShadingTileOverscan.GetValueOnRenderThread(), 2, GroupTileSize / 2 - 1);
+		const int32 TileOverscan = FMath::Clamp(CVarTSRShadingTileOverscan.GetValueOnRenderThread(), 3, GroupTileSize / 2 - 1);
 		const int32 TileSize = GroupTileSize - 2 * TileOverscan;
 
 		FTSRRejectShadingCS::FParameters* PassParameters = GraphBuilder.AllocParameters<FTSRRejectShadingCS::FParameters>();
