Interchange: Fix config file not at the correct location
#rnx

diff --git a/Engine/Config/BaseEngine.ini b/Engine/Config/BaseEngine.ini
--- a/Engine/Config/BaseEngine.ini
+++ b/Engine/Config/BaseEngine.ini
@@ -3316,4 +3316,11 @@ bUpdateBuildGraphPropertiesFile=false
 +RequiresCompleteValueForPostConstruction=/Script/Engine.BodyInstance:CollisionResponses
 
 [Ias]
-EnableInEditor=False
\ No newline at end of file
+EnableInEditor=False
+
+[/Script/InterchangeEngine.InterchangeProjectSettings]
+ContentImportSettings=(DefaultPipelineStackOverride=((Materials, "Materials"),(Textures, "Textures")),ShowPipelineStacksConfigurationDialogOverride=((Textures, False)),PipelineStacks=(("Assets", (Pipelines=("/Interchange/Pipelines/DefaultAssetsPipeline.DefaultAssetsPipeline"),PerTranslatorPipelines=((Translator="/Script/InterchangeImport.InterchangeMaterialXTranslator",Pipelines=("/Interchange/Pipelines/DefaultAssetsPipeline.DefaultAssetsPipeline","/Interchange/Pipelines/DefaultMaterialXPipeline.DefaultMaterialXPipeline")),(Translator="/Script/InterchangeImport.InterchangeGltfTranslator",Pipelines=("/Interchange/Pipelines/DefaultGLTFAssetsPipeline.DefaultGLTFAssetsPipeline","/Interchange/Pipelines/DefaultGLTFPipeline.DefaultGLTFPipeline"))))),("Materials", (Pipelines=("/Interchange/Pipelines/DefaultMaterialPipeline.DefaultMaterialPipeline"),PerTranslatorPipelines=((Translator="/Script/InterchangeImport.InterchangeMaterialXTranslator",Pipelines=("/Interchange/Pipelines/DefaultMaterialPipeline.DefaultMaterialPipeline","/Interchange/Pipelines/DefaultMaterialXPipeline.DefaultMaterialXPipeline"))))),("Textures", (Pipelines=("/Interchange/Pipelines/DefaultTexturePipeline.DefaultTexturePipeline")))),DefaultPipelineStack="Assets",PipelineConfigurationDialogClass="/Script/InterchangeEditorPipelines.InterchangePipelineConfigurationGeneric",bShowPipelineStacksConfigurationDialog=True)
+SceneImportSettings=(PipelineStacks=(("Scene", (Pipelines=("/Interchange/Pipelines/DefaultSceneAssetsPipeline.DefaultSceneAssetsPipeline","/Interchange/Pipelines/DefaultSceneLevelPipeline.DefaultSceneLevelPipeline"),PerTranslatorPipelines=((Translator="/Script/InterchangeImport.InterchangeMaterialXTranslator",Pipelines=("/Interchange/Pipelines/DefaultSceneAssetsPipeline.DefaultSceneAssetsPipeline","/Interchange/Pipelines/DefaultSceneLevelPipeline.DefaultSceneLevelPipeline","/Interchange/Pipelines/DefaultMaterialXPipeline.DefaultMaterialXPipeline")),(Translator="/Script/InterchangeImport.InterchangeGltfTranslator",Pipelines=("/Interchange/Pipelines/DefaultGLTFSceneAssetsPipeline.DefaultGLTFSceneAssetsPipeline","/Interchange/Pipelines/DefaultSceneLevelPipeline.DefaultSceneLevelPipeline","/Interchange/Pipelines/DefaultGLTFPipeline.DefaultGLTFPipeline")))))),DefaultPipelineStack="Scene",PipelineConfigurationDialogClass="/Script/InterchangeEditorPipelines.InterchangePipelineConfigurationGeneric",bShowPipelineStacksConfigurationDialog=True)
+FilePickerClass=/Script/InterchangeEditorUtilities.InterchangeFilePickerGeneric
+bStaticMeshUseSmoothEdgesIfSmoothingInformationIsMissing=True
+GenericPipelineClass=/Script/InterchangePipelines.InterchangeGenericAssetsPipeline
\ No newline at end of file
diff --git a/Engine/Source/Runtime/Interchange/Engine/Public/InterchangeProjectSettings.h b/Engine/Source/Runtime/Interchange/Engine/Public/InterchangeProjectSettings.h
--- a/Engine/Source/Runtime/Interchange/Engine/Public/InterchangeProjectSettings.h
+++ b/Engine/Source/Runtime/Interchange/Engine/Public/InterchangeProjectSettings.h
@@ -74,7 +74,7 @@ struct FInterchangeContentImportSettings : public FInterchangeImportSettings
 	TMap<EInterchangeTranslatorAssetType, bool> ShowPipelineStacksConfigurationDialogOverride;
 };
 
-UCLASS(config=Interchange, meta=(DisplayName=Interchange), MinimalAPI)
+UCLASS(config=Engine, meta=(DisplayName=Interchange), MinimalAPI)
 class UInterchangeProjectSettings : public UDeveloperSettings
 {
 	GENERATED_BODY()
