EditorDomain: Disable the EditorDomainOplog for new incremental cooking until we have fixed initialization to occur at cook startup rather than too early in engine startup.
#rnx

diff --git a/Engine/Config/BaseEditor.ini b/Engine/Config/BaseEditor.ini
--- a/Engine/Config/BaseEditor.ini
+++ b/Engine/Config/BaseEditor.ini
@@ -408,6 +408,7 @@ CookProcessCount=1
 [EditorDomain]
 EditorDomainEnabled=false
 SaveUnversioned=true
+CookAttachmentsEnabled=false
 ; +ClassBlockList=/Long/PackageName.ClassName
 ; +PackageBlockList=/Long/PackageName OR ../../../Standard/Filename/NoExtension OR D:\Absolute\Filename\NoExtension
 
