Fix Mali device detection potentially incorrectly matching substrings of other devices.

diff --git a/Engine/Config/BaseDeviceProfiles.ini b/Engine/Config/BaseDeviceProfiles.ini
--- a/Engine/Config/BaseDeviceProfiles.ini
+++ b/Engine/Config/BaseDeviceProfiles.ini
@@ -692,14 +692,14 @@ BaseProfileName=VisionPro
 +MatchProfile=(Profile="Android_Mali_T7xx",Match=((SourceType=SRC_GpuFamily,CompareType=CMP_Regex,MatchString="^Mali\\-T7")))
 +MatchProfile=(Profile="Android_Mali_T8xx",Match=((SourceType=SRC_GpuFamily,CompareType=CMP_Regex,MatchString="^Mali\\-T8")))
 +MatchProfile=(Profile="Android_Mali_G71",Match=((SourceType=SRC_GpuFamily,CompareType=CMP_Regex,MatchString="^Mali\\-G71($| )")))
-+MatchProfile=(Profile="Android_Mali_G72_Vulkan",Match=((SourceType=SRC_GpuFamily,CompareType=CMP_Regex,MatchString="^Mali\\-G72"),(SourceType=SRC_AndroidVersion, CompareType=CMP_Regex,MatchString="([0-9]+).*"),(SourceType=SRC_PreviousRegexMatch,CompareType=CMP_GreaterEqual,MatchString="9"),(SourceType=SRC_VulkanAvailable,CompareType=CMP_Equal,MatchString="true")))
-+MatchProfile=(Profile="Android_Mali_G72",Match=((SourceType=SRC_GpuFamily,CompareType=CMP_Regex,MatchString="^Mali\\-G72")))
-+MatchProfile=(Profile="Android_Mali_G76_Vulkan",Match=((SourceType=SRC_GpuFamily,CompareType=CMP_Regex,MatchString="^Mali\\-G76"),(SourceType=SRC_VulkanAvailable,CompareType=CMP_Equal,MatchString="true")))
-+MatchProfile=(Profile="Android_Mali_G76",Match=((SourceType=SRC_GpuFamily,CompareType=CMP_Regex,MatchString="^Mali\\-G76")))
-+MatchProfile=(Profile="Android_Mali_G77_Vulkan",Match=((SourceType=SRC_GpuFamily,CompareType=CMP_Regex,MatchString="^Mali\\-G77"),(SourceType=SRC_VulkanAvailable,CompareType=CMP_Equal,MatchString="true")))
-+MatchProfile=(Profile="Android_Mali_G77",Match=((SourceType=SRC_GpuFamily,CompareType=CMP_Regex,MatchString="^Mali\\-G77")))
-+MatchProfile=(Profile="Android_Mali_G78_Vulkan",Match=((SourceType=SRC_GpuFamily,CompareType=CMP_Regex,MatchString="^Mali\\-G78"),(SourceType=SRC_VulkanAvailable,CompareType=CMP_Equal,MatchString="true")))
-+MatchProfile=(Profile="Android_Mali_G78",Match=((SourceType=SRC_GpuFamily,CompareType=CMP_Regex,MatchString="^Mali\\-G78")))
++MatchProfile=(Profile="Android_Mali_G72_Vulkan",Match=((SourceType=SRC_GpuFamily,CompareType=CMP_Regex,MatchString="^Mali\\-G72($| )"),(SourceType=SRC_AndroidVersion, CompareType=CMP_Regex,MatchString="([0-9]+).*"),(SourceType=SRC_PreviousRegexMatch,CompareType=CMP_GreaterEqual,MatchString="9"),(SourceType=SRC_VulkanAvailable,CompareType=CMP_Equal,MatchString="true")))
++MatchProfile=(Profile="Android_Mali_G72",Match=((SourceType=SRC_GpuFamily,CompareType=CMP_Regex,MatchString="^Mali\\-G72($| )")))
++MatchProfile=(Profile="Android_Mali_G76_Vulkan",Match=((SourceType=SRC_GpuFamily,CompareType=CMP_Regex,MatchString="^Mali\\-G76($| )"),(SourceType=SRC_VulkanAvailable,CompareType=CMP_Equal,MatchString="true")))
++MatchProfile=(Profile="Android_Mali_G76",Match=((SourceType=SRC_GpuFamily,CompareType=CMP_Regex,MatchString="^Mali\\-G76($| )")))
++MatchProfile=(Profile="Android_Mali_G77_Vulkan",Match=((SourceType=SRC_GpuFamily,CompareType=CMP_Regex,MatchString="^Mali\\-G77($| )"),(SourceType=SRC_VulkanAvailable,CompareType=CMP_Equal,MatchString="true")))
++MatchProfile=(Profile="Android_Mali_G77",Match=((SourceType=SRC_GpuFamily,CompareType=CMP_Regex,MatchString="^Mali\\-G77($| )")))
++MatchProfile=(Profile="Android_Mali_G78_Vulkan",Match=((SourceType=SRC_GpuFamily,CompareType=CMP_Regex,MatchString="^Mali\\-G78($| )"),(SourceType=SRC_VulkanAvailable,CompareType=CMP_Equal,MatchString="true")))
++MatchProfile=(Profile="Android_Mali_G78",Match=((SourceType=SRC_GpuFamily,CompareType=CMP_Regex,MatchString="^Mali\\-G78($| )")))
 +MatchProfile=(Profile="Android_Mali_G710_Vulkan",Match=((SourceType=SRC_GpuFamily,CompareType=CMP_Regex,MatchString="^Mali\\-G710"),(SourceType=SRC_VulkanAvailable,CompareType=CMP_Equal,MatchString="true")))
 +MatchProfile=(Profile="Android_Mali_G710",Match=((SourceType=SRC_GpuFamily,CompareType=CMP_Regex,MatchString="^Mali\\-G710")))
 +MatchProfile=(Profile="Android_Mali_G7xx_Vulkan",Match=((SourceType=SRC_GpuFamily,CompareType=CMP_Regex,MatchString="^Mali\\-G7[0-9][0-9]"),(SourceType=SRC_VulkanAvailable,CompareType=CMP_Equal,MatchString="true")))
