Merge PR from Oculus to add Quest 3 to BaseDeviceProfiles.ini
Required for Quest 3 MR + MetaXR plugin and newer VRS mobile foveation work.

diff --git a/Engine/Config/BaseDeviceProfiles.ini b/Engine/Config/BaseDeviceProfiles.ini
--- a/Engine/Config/BaseDeviceProfiles.ini
+++ b/Engine/Config/BaseDeviceProfiles.ini
@@ -123,6 +123,7 @@
 +DeviceProfileNameAndTypes=Android_PowerVR_GM9xxx_Vulkan,Android
 +DeviceProfileNameAndTypes=Android_TegraK1,Android
 +DeviceProfileNameAndTypes=Android_Unknown_Vulkan,Android
++DeviceProfileNameAndTypes=Meta_Quest_3,Android
 +DeviceProfileNameAndTypes=Meta_Quest_Pro,Android
 +DeviceProfileNameAndTypes=Oculus_Quest2,Android
 +DeviceProfileNameAndTypes=Oculus_Quest,Android
@@ -663,6 +664,7 @@ BaseProfileName=VisionPro
 ; These rules are checked in turn and stop at the profile that matches all criteria
 [/Script/AndroidDeviceProfileSelector.AndroidDeviceProfileMatchingRules]
 ; Match Quest devices directly
++MatchProfile=(Profile="Meta_Quest_3",Match=((SourceType=SRC_DeviceMake,CompareType=CMP_Equal,MatchString="Oculus"), (SourceType=SRC_HMDSystemName,CompareType=CMP_Equal,MatchString="Meta Quest 3")))
 +MatchProfile=(Profile="Meta_Quest_Pro",Match=((SourceType=SRC_DeviceMake,CompareType=CMP_Equal,MatchString="Oculus"), (SourceType=SRC_HMDSystemName,CompareType=CMP_Equal,MatchString="Meta Quest Pro")))
 +MatchProfile=(Profile="Oculus_Quest2",Match=((SourceType=SRC_DeviceMake,CompareType=CMP_Equal,MatchString="Oculus"), (SourceType=SRC_HMDSystemName,CompareType=CMP_Equal,MatchString="Oculus Quest2")))
 +MatchProfile=(Profile="Oculus_Quest",Match=((SourceType=SRC_DeviceMake,CompareType=CMP_Equal,MatchString="Oculus")))
@@ -1026,6 +1028,10 @@ BaseProfileName=Android_High
 ; enable Vulkan
 +CVars=r.Android.DisableVulkanSupport=0
 
+[Meta_Quest_3 DeviceProfile]
+DeviceType=Android
+BaseProfileName=Meta_Quest_Pro
+
 [Meta_Quest_Pro DeviceProfile]
 DeviceType=Android
 BaseProfileName=Oculus_Quest2
@@ -1043,6 +1049,7 @@ BaseProfileName=Android_Mid
 ; disable running Niagara compute shaders on the GPU. This can cause hard to trace crashes.
 +CVars=fx.NiagaraAllowGPUParticles=0
 +CVars=FX.AllowGPUSorting=0
++CVars=r.Vulkan.VRSFormat=3
 
 [Mac DeviceProfile]
 DeviceType=Mac
