Disable r.Lumen.ScreenProbeGather.ShortRangeAO.HardwareRayTracing on Cinematic GI quality until artifacts with masked materials and self-intersection on Nanite meshes can be fixed

diff --git a/Engine/Config/BaseScalability.ini b/Engine/Config/BaseScalability.ini
--- a/Engine/Config/BaseScalability.ini
+++ b/Engine/Config/BaseScalability.ini
@@ -331,7 +331,7 @@ r.Lumen.ScreenProbeGather.StochasticInterpolation=0
 r.Lumen.ScreenProbeGather.FullResolutionJitterWidth=.5
 r.Lumen.ScreenProbeGather.TwoSidedFoliageBackfaceDiffuse=1
 r.Lumen.ScreenProbeGather.ScreenTraces.HZBTraversal.FullResDepth=1
-r.Lumen.ScreenProbeGather.ShortRangeAO.HardwareRayTracing=1
+r.Lumen.ScreenProbeGather.ShortRangeAO.HardwareRayTracing=0
 r.Lumen.TranslucencyVolume.GridPixelSize=32
 r.Lumen.TranslucencyVolume.TraceFromVolume=1
 r.Lumen.TranslucencyVolume.TracingOctahedronResolution=4
