AssetManager: Update some comments for PrimaryAsset API.
#rnx

diff --git a/Engine/Config/BaseGame.ini b/Engine/Config/BaseGame.ini
--- a/Engine/Config/BaseGame.ini
+++ b/Engine/Config/BaseGame.ini
@@ -219,6 +219,14 @@ FrameInterpolationProcessor=Class'/Script/LiveLink.LiveLinkBasicFrameInterpolati
 
 [/Script/Engine.AssetManagerSettings]
 @PrimaryAssetTypesToScan=PrimaryAssetType
+
+; Note the default rule here for both PrimaryAssetLabels and Maps sets them to bIsEditorOnly=true, which means the
+; AssetManager will NOT use them to pull assets into the build; they will only label the target chunks (aka pak files)
+; into which assets pulled in by other sources (e.g. ProjectSettings) will be put. It also sets the CookRule=Unknown,
+; which is another way of saying that they do not impact whether the assets labeled by the PrimaryAssetLabel or Map are
+; cooked. If you want PrimaryAssetLabels or maps to cause assets to be cooked even if the label and its assets are otherwise
+; unreferenced, then set bIsEditorOnly=false and CookRule=AlwaysCook. 
+
 +PrimaryAssetTypesToScan=(PrimaryAssetType="Map",AssetBaseClass=/Script/Engine.World,bHasBlueprintClasses=False,bIsEditorOnly=True,Directories=((Path="/Game/Maps")),SpecificAssets=,Rules=(Priority=-1,ChunkId=-1,bApplyRecursively=True,CookRule=Unknown))
 +PrimaryAssetTypesToScan=(PrimaryAssetType="PrimaryAssetLabel",AssetBaseClass=/Script/Engine.PrimaryAssetLabel,bHasBlueprintClasses=False,bIsEditorOnly=True,Directories=((Path="/Game")),SpecificAssets=,Rules=(Priority=-1,ChunkId=-1,bApplyRecursively=True,CookRule=Unknown))
 
diff --git a/Engine/Source/Runtime/Engine/Classes/Engine/AssetManagerTypes.h b/Engine/Source/Runtime/Engine/Classes/Engine/AssetManagerTypes.h
--- a/Engine/Source/Runtime/Engine/Classes/Engine/AssetManagerTypes.h
+++ b/Engine/Source/Runtime/Engine/Classes/Engine/AssetManagerTypes.h
@@ -162,7 +162,12 @@ public:
 	UPROPERTY(EditAnywhere, Category = AssetType)
 	bool bHasBlueprintClasses;
 
-	/** True if this type is editor only. Disable this to consider content that is part of the engine or engine plugins. */
+	/**
+	 * If true this type will not cause anything to be cooked; the AssetManager will use instances of this type to
+	 * define chunk assignments and NeverCook rules, but will ignore AlwaysCook rules. Assets labeled by instances
+	 * of this type will need to be reference by another PrimaryAsset, or by something outside the AssetManager,
+	 * to be cooked.
+	 */
 	UPROPERTY(EditAnywhere, Category = AssetType)
 	bool bIsEditorOnly;
 
diff --git a/Engine/Source/Runtime/Engine/Classes/Engine/PrimaryAssetLabel.h b/Engine/Source/Runtime/Engine/Classes/Engine/PrimaryAssetLabel.h
--- a/Engine/Source/Runtime/Engine/Classes/Engine/PrimaryAssetLabel.h
+++ b/Engine/Source/Runtime/Engine/Classes/Engine/PrimaryAssetLabel.h
@@ -35,7 +35,7 @@ public:
 	UPROPERTY(EditAnywhere, Category = PrimaryAssetLabel)
 	uint32 bIsRuntimeLabel : 1;
 
-	/** Set to true if the label asset itself should be cooked and available at runtime. This does not affect the assets that are labeled, they are set with cook rule */
+	/** If true, redirectors found by the AssetLabel's explicit assets or directory will be labeled. If false, redirectors will be skipped. */
 	UPROPERTY(EditAnywhere, Category = PrimaryAssetLabel)
 	uint32 bIncludeRedirectors : 1;
 
@@ -52,7 +52,7 @@ public:
 	FCollectionReference AssetCollection;
 
 	/** Set to editor only if this is not available in a cooked build */
-	virtual bool IsEditorOnly() const
+	virtual bool IsEditorOnly() const override
 	{
 		return !bIsRuntimeLabel && !IsTemplate();
 	}
