Disable moving resolution LOD bias for directional lights by default as it can cause artifacts and flicker with certain types of Nanite geometry

diff --git a/Engine/Config/BaseScalability.ini b/Engine/Config/BaseScalability.ini
--- a/Engine/Config/BaseScalability.ini
+++ b/Engine/Config/BaseScalability.ini
@@ -145,7 +145,7 @@ r.LightMaxDrawDistanceScale=0
 r.CapsuleShadows=0
 r.Shadow.Virtual.MaxPhysicalPages=512
 r.Shadow.Virtual.ResolutionLodBiasDirectional=0.0
-r.Shadow.Virtual.ResolutionLodBiasDirectionalMoving=1.0
+r.Shadow.Virtual.ResolutionLodBiasDirectionalMoving=0.0
 r.Shadow.Virtual.ResolutionLodBiasLocal=1.0
 r.Shadow.Virtual.ResolutionLodBiasLocalMoving=2.0
 r.Shadow.Virtual.SMRT.RayCountDirectional=0
@@ -169,7 +169,7 @@ r.LightMaxDrawDistanceScale=.5
 r.CapsuleShadows=1
 r.Shadow.Virtual.MaxPhysicalPages=512
 r.Shadow.Virtual.ResolutionLodBiasDirectional=0.0
-r.Shadow.Virtual.ResolutionLodBiasDirectionalMoving=1.0
+r.Shadow.Virtual.ResolutionLodBiasDirectionalMoving=0.0
 r.Shadow.Virtual.ResolutionLodBiasLocal=1.0
 r.Shadow.Virtual.ResolutionLodBiasLocalMoving=2.0
 r.Shadow.Virtual.SMRT.RayCountDirectional=4
@@ -196,7 +196,7 @@ r.LightMaxDrawDistanceScale=1
 r.CapsuleShadows=1
 r.Shadow.Virtual.MaxPhysicalPages=2048
 r.Shadow.Virtual.ResolutionLodBiasDirectional=-0.5
-r.Shadow.Virtual.ResolutionLodBiasDirectionalMoving=0.5
+r.Shadow.Virtual.ResolutionLodBiasDirectionalMoving=-0.5
 r.Shadow.Virtual.ResolutionLodBiasLocal=0.0
 r.Shadow.Virtual.ResolutionLodBiasLocalMoving=1.0
 r.Shadow.Virtual.SMRT.RayCountDirectional=8
@@ -223,7 +223,7 @@ r.LightMaxDrawDistanceScale=1
 r.CapsuleShadows=1
 r.Shadow.Virtual.MaxPhysicalPages=4096
 r.Shadow.Virtual.ResolutionLodBiasDirectional=-1.5
-r.Shadow.Virtual.ResolutionLodBiasDirectionalMoving=0.0
+r.Shadow.Virtual.ResolutionLodBiasDirectionalMoving=-1.5
 r.Shadow.Virtual.ResolutionLodBiasLocal=0.0
 r.Shadow.Virtual.ResolutionLodBiasLocalMoving=1.0
 r.Shadow.Virtual.SMRT.RayCountDirectional=8
