Add all the new iPhone15 (Plus/Pro/ProMax) models

diff --git a/Engine/Config/BaseDeviceProfiles.ini b/Engine/Config/BaseDeviceProfiles.ini
--- a/Engine/Config/BaseDeviceProfiles.ini
+++ b/Engine/Config/BaseDeviceProfiles.ini
@@ -45,6 +45,10 @@
 +DeviceProfileNameAndTypes=iPhone14Plus,IOS
 +DeviceProfileNameAndTypes=iPhone14Pro,IOS
 +DeviceProfileNameAndTypes=iPhone14ProMax,IOS
++DeviceProfileNameAndTypes=iPhone15,IOS
++DeviceProfileNameAndTypes=iPhone15Plus,IOS
++DeviceProfileNameAndTypes=iPhone15Pro,IOS
++DeviceProfileNameAndTypes=iPhone15ProMax,IOS
 +DeviceProfileNameAndTypes=iPadPro,IOS
 +DeviceProfileNameAndTypes=iPadPro105,IOS
 +DeviceProfileNameAndTypes=iPadPro129,IOS
@@ -216,6 +220,10 @@ iPhone14,7=iPhone14
 iPhone14,8=iPhone14Plus
 iPhone15,2=iPhone14Pro
 iPhone15,3=iPhone14ProMax
+iPhone15,4=iPhone15
+iPhone15,5=iPhone15Plus
+iPhone16,1=iPhone15Pro
+iPhone16,2=iPhone15ProMax
 iPod9,1=iPodTouch7
 iPad5,[1,2]=iPadMini4
 iPad5,[3,4]=iPadAir2
@@ -254,7 +262,7 @@ AppleTV5,=AppleTV
 AppleTV6,=AppleTV4K
 AppleTV11,=AppleTV2_4K
 ; Below are generic fallbacks for unrecognized devices. Add any specific devices above this section.
-iPhone=iPhone14
+iPhone=iPhone15
 iPod=iPodTouch7
 iPad=iPad10
 AppleTV=AppleTV2_4K
@@ -471,6 +479,24 @@ DeviceType=IOS
 BaseProfileName=iPhone13ProMax
 +CVars=ios.PhysicalScreenDensity=460
 
+[iPhone15 DeviceProfile]
+DeviceType=IOS
+BaseProfileName=iPhone13
+
+[iPhone15Plus DeviceProfile]
+DeviceType=IOS
+BaseProfileName=iPhone13
++CVars=ios.PhysicalScreenDensity=458
+
+[iPhone15Pro DeviceProfile]
+DeviceType=IOS
+BaseProfileName=iPhone13Pro
+
+[iPhone15ProMax DeviceProfile]
+DeviceType=IOS
+BaseProfileName=iPhone13ProMax
++CVars=ios.PhysicalScreenDensity=460
+
 [iPodTouch7 DeviceProfile]
 DeviceType=IOS
 BaseProfileName=iPhone7
