IterativeCooking: Temporarily disallow Materials to be iteratively skipped in incremental cooks or in iterative++ cooks. If hlsl shader files change the materials need to be recooked, but the iterative and incremental systems do not yet detect that dependency.
#rnx

diff --git a/Engine/Config/BaseEditor.ini b/Engine/Config/BaseEditor.ini
--- a/Engine/Config/BaseEditor.ini
+++ b/Engine/Config/BaseEditor.ini
@@ -459,7 +459,7 @@ SaveUnversioned=true
 +IterativeScriptPackageAllowList=/Script/UnrealEd
 
 +IterativeScriptPackageAllowList=/Script/AudioEditor
-
++IterativeClassBlockList=/Script/Engine.MaterialInterface
 
 [CookPlatformDataCacheSettings]
 ; Add <ShortClassName>=<Integer> here to limit the number of instances that can have BeginCacheForCookedPlatformData called at once
