Enables TSR history resurrection on medium anti-aliasing scalability

#tests ps5, xboxsx, win64

diff --git a/Engine/Config/BaseScalability.ini b/Engine/Config/BaseScalability.ini
--- a/Engine/Config/BaseScalability.ini
+++ b/Engine/Config/BaseScalability.ini
@@ -69,7 +69,7 @@ r.TSR.History.UpdateQuality=1
 r.TSR.ShadingRejection.Flickering=0
 r.TSR.Velocity.Extrapolation=0
 r.TSR.RejectionAntiAliasingQuality=1
-r.TSR.Resurrection=0
+r.TSR.Resurrection=1
 
 [AntiAliasingQuality@2]
 r.FXAA.Quality=3
