Adding redirect for ControlRig refactor to BaseControlRig.

diff --git a/Engine/Config/BaseEngine.ini b/Engine/Config/BaseEngine.ini
--- a/Engine/Config/BaseEngine.ini
+++ b/Engine/Config/BaseEngine.ini
@@ -1212,6 +1212,9 @@ DefaultWorldPartitionHLODLayer=/Engine/Maps/Templates/HLODs/HLODLayer_Instanced.
 +ClassRedirects=(OldName="/Script/Engine.MaterialExpressionStrataThinFilm",						NewName="/Script/Engine.MaterialExpressionSubstrateThinFilm")
 +StructRedirects=(OldName="/Script/Engine.StrataMaterialInput",									NewName="/Script/Engine.SubstrateMaterialInput")
 
+; Control Rig refactor
++ClassRedirects=(OldName="/Script/ControlRig.ControlRig", NewName="/Script/ControlRig.BaseControlRig")
+
 ; End of releases redirects
 
 
