Oodle WAVE: enable for use

diff --git a/Engine/Config/BaseEditor.ini b/Engine/Config/BaseEditor.ini
--- a/Engine/Config/BaseEditor.ini
+++ b/Engine/Config/BaseEditor.ini
@@ -255,7 +255,7 @@ bExpandClassPickerClassList=False
 Class=UnrealEd.ActorFactory
 
 [AudioImporter]
-EnableOodleWAV=False
+EnableOodleWAV=True
 
 [TextureImporter]
 AllowNonPowerOfTwoTextures=True
