Added the following to ConsoleVariables.ini:
r.Nanite.AllowTessellation=1
r.Nanite.Tessellation=1

diff --git a/Engine/Config/ConsoleVariables.ini b/Engine/Config/ConsoleVariables.ini
--- a/Engine/Config/ConsoleVariables.ini
+++ b/Engine/Config/ConsoleVariables.ini
@@ -31,6 +31,8 @@
 
 [Startup]
 
+r.Nanite.AllowTessellation=1
+r.Nanite.Tessellation=1
 
 ; Uncomment to get detailed logs on shader compiles and the opportunity to retry on errors
 ;r.ShaderDevelopmentMode=1
