Add Xclipse 940 chipset and GPU (5xx supported and falls through to low)

#android

diff --git a/Engine/Config/BaseDeviceProfiles.ini b/Engine/Config/BaseDeviceProfiles.ini
--- a/Engine/Config/BaseDeviceProfiles.ini
+++ b/Engine/Config/BaseDeviceProfiles.ini
@@ -117,8 +117,11 @@
 +DeviceProfileNameAndTypes=Android_Mali_G710_Vulkan,Android
 +DeviceProfileNameAndTypes=Android_Mali_G7xx,Android
 +DeviceProfileNameAndTypes=Android_Mali_G7xx_Vulkan,Android
-+DeviceProfileNameAndTypes=Android_Xclipse_920,Android
-+DeviceProfileNameAndTypes=Android_Xclipse_920_Vulkan,Android
++DeviceProfileNameAndTypes=Android_Xclipse_5xx,Android
++DeviceProfileNameAndTypes=Android_Xclipse_5xx_Vulkan,Android
++DeviceProfileNameAndTypes=Android_Xclipse_9xx,Android
++DeviceProfileNameAndTypes=Android_Xclipse_xxx,Android
++DeviceProfileNameAndTypes=Android_Xclipse_xxx_Vulkan,Android
 +DeviceProfileNameAndTypes=Android_Vulkan_SM5,Android
 +DeviceProfileNameAndTypes=Android_PowerVR_G6xxx,Android
 +DeviceProfileNameAndTypes=Android_PowerVR_GT7xxx,Android
@@ -738,8 +741,12 @@ BaseProfileName=VisionPro
 ; Samsung Exynos / Xclipse GPU-based devices use Vulkan by default.
 ;+MatchProfile=(Profile="Android_Vulkan_SM5",Match=((SourceType=SRC_GpuFamily,CompareType=CMP_Regex,MatchString="Samsung Xclipse 920"),(SourceType=SRC_VulkanAvailable,CompareType=CMP_Equal,MatchString="true")))
 ;+MatchProfile=(Profile="Android_Vulkan_SM5_Mali",Match=((SourceType=SRC_GpuFamily,CompareType=CMP_Regex,MatchString="^Mali\\-G7[67]"),(SourceType=SRC_VulkanAvailable,CompareType=CMP_Equal,MatchString="true")))
-+MatchProfile=(Profile="Android_Xclipse_920_Vulkan",Match=((SourceType=SRC_GpuFamily,CompareType=CMP_Regex,MatchString="Samsung Xclipse 920"),(SourceType=SRC_VulkanAvailable,CompareType=CMP_Equal,MatchString="true")))
-+MatchProfile=(Profile="Android_Xclipse_920",Match=((SourceType=SRC_GpuFamily,CompareType=CMP_Regex,MatchString="Samsung Xclipse 920")))
++MatchProfile=(Profile="Android_Xclipse_5xx_Vulkan",Match=((SourceType=SRC_GpuFamily,CompareType=CMP_Regex,MatchString="Samsung Xclipse 5[0-9][0-9]"),(SourceType=SRC_VulkanAvailable,CompareType=CMP_Equal,MatchString="true")))
++MatchProfile=(Profile="Android_Xclipse_5xx",Match=((SourceType=SRC_GpuFamily,CompareType=CMP_Regex,MatchString="Samsung Xclipse 5[0-9][0-9]")))
++MatchProfile=(Profile="Android_Xclipse_9xx_Vulkan",Match=((SourceType=SRC_GpuFamily,CompareType=CMP_Regex,MatchString="Samsung Xclipse 9[0-9][0-9]"),(SourceType=SRC_VulkanAvailable,CompareType=CMP_Equal,MatchString="true")))
++MatchProfile=(Profile="Android_Xclipse_9xx",Match=((SourceType=SRC_GpuFamily,CompareType=CMP_Regex,MatchString="Samsung Xclipse 9[0-9][0-9]")))
++MatchProfile=(Profile="Android_Xclipse_xxx_Vulkan",Match=((SourceType=SRC_GpuFamily,CompareType=CMP_Regex,MatchString="Samsung Xclipse [0-9][0-9][0-9]"),(SourceType=SRC_VulkanAvailable,CompareType=CMP_Equal,MatchString="true")))
++MatchProfile=(Profile="Android_Xclipse_xxx",Match=((SourceType=SRC_GpuFamily,CompareType=CMP_Regex,MatchString="Samsung Xclipse [0-9][0-9][0-9]")))
 
 ; PowerVR based devices
 +MatchProfile=(Profile="Android_PowerVR_G6xxx",Match=((SourceType=SRC_GpuFamily,CompareType=CMP_Regex,MatchString="PowerVR Rogue G6[0-9]+")))
@@ -917,15 +924,39 @@ BaseProfileName=Android_Adreno7xx
 ;----------------
 ; Xclipse devices
 
-[Android_Xclipse_920 DeviceProfile]
+[Android_Xclipse_5xx DeviceProfile]
+DeviceType=Android
+BaseProfileName=Android_Low
+
+[Android_Xclipse_5xx_Vulkan DeviceProfile]
+DeviceType=Android
+BaseProfileName=Android_Xclipse_5xx
+; enable Vulkan
++CVars=r.Android.DisableVulkanSupport=0
+
+[Android_Xclipse_9xx DeviceProfile]
 DeviceType=Android
 BaseProfileName=Android_High
 
-[Android_Xclipse_920_Vulkan DeviceProfile]
+[Android_Xclipse_9xx_Vulkan DeviceProfile]
 DeviceType=Android
-BaseProfileName=Android_Xclipse_920
+BaseProfileName=Android_Xclipse_9xx
 ; enable Vulkan
 +CVars=r.Android.DisableVulkanSupport=0
+; Disable Depth24bit 
++CVars=r.Vulkan.Depth24Bit=0
+
+[Android_Xclipse_xxx DeviceProfile]
+DeviceType=Android
+BaseProfileName=Android_Mid
+
+[Android_Xclipse_xxx_Vulkan DeviceProfile]
+DeviceType=Android
+BaseProfileName=Android_Xclipse_xxx
+; Enable Vulkan
++CVars=r.Android.DisableVulkanSupport=0
+; Disable Depth24bit 
++CVars=r.Vulkan.Depth24Bit=0
 
 ;----------------
 ; PowerVR devices
