- Enabled the system to massively reduce what settings are saved in Saved/Config. Note that *Editor* and *User* ini files are not changed, they always fully save. Additionally, sections can be opted in for a particular ini file (like Game, Engine, etc). but adding the section name to the [SectionsToSave] section in a file in the hierarchy (like BaseEngine.ini or DefaultEngine.ini)

diff --git a/Engine/Config/Base.ini b/Engine/Config/Base.ini
--- a/Engine/Config/Base.ini
+++ b/Engine/Config/Base.ini
@@ -1,5 +1,9 @@
 ; Base.ini is required to exist
 
+; To always save all sections like before, set bCanSaveAllSections to true
+; You can set this in a particular ini file as well, like BaseEngine.ini
+; or add more Sections that you would like to be completely saved out
+; See BaseEngine.ini for more examples
 [SectionsToSave]
-bCanSaveAllSections=true
+bCanSaveAllSections=false
 +Section=CurrentIniVersion
