BaseEditorKeyBindings.ini - Explicitly save out user defined keybinds to fix a bug where they were not persisting after a restart.

diff --git a/Engine/Config/BaseEditorKeyBindings.ini b/Engine/Config/BaseEditorKeyBindings.ini
--- a/Engine/Config/BaseEditorKeyBindings.ini
+++ b/Engine/Config/BaseEditorKeyBindings.ini
@@ -1,3 +1,6 @@
+[SectionsToSave]
++Section=UserDefinedChords
+
 [/Script/UnrealEd.UnrealEdKeyBindings]
 ; --- Curve Editor Bindings
 +KeyBindings=(bCtrlDown=False,bAltDown=False,bShiftDown=False,Key="one",CommandName="CurveEditor_ChangeInterpModeAUTO")
