Explicitly ask PIE console command history and asset picker settings to be saved into ini configs

diff --git a/Engine/Config/BaseEditorSettings.ini b/Engine/Config/BaseEditorSettings.ini
--- a/Engine/Config/BaseEditorSettings.ini
+++ b/Engine/Config/BaseEditorSettings.ini
@@ -2,6 +2,9 @@
 ; These settings are overridden by a project's DefaultEditorSettings.ini file, per-user editor setting defaults are in *EditorPerProjectUserSettings.ini
 ; Most of these options are modifiable from Project Settings in the editor, and are not available in a packaged build
 
+[SectionsToSave]
++Section=/Script/UnrealEd.ContentBrowserSettings
+
 [Internationalization]
 ShouldUseLocalizedNumericInput=True
 ShouldUseLocalizedPropertyNames=True
diff --git a/Engine/Config/BaseInput.ini b/Engine/Config/BaseInput.ini
--- a/Engine/Config/BaseInput.ini
+++ b/Engine/Config/BaseInput.ini
@@ -1,3 +1,6 @@
+[SectionsToSave]
++Section=/Script/Engine.Console
+
 [/Script/Engine.InputSettings]
 bEnableMouseSmoothing=true
 bEnableFOVScaling=true
