- Fixed the ini array for saving out the Versioning section in all .ini's. The Base.ini is not merged, so it does not expect the + commands (the + in this case was part of the key, which would break versioning). Note that versioning is much less useful with the new "save almost nothing" code

diff --git a/Engine/Config/Base.ini b/Engine/Config/Base.ini
--- a/Engine/Config/Base.ini
+++ b/Engine/Config/Base.ini
@@ -6,4 +6,4 @@
 ; See BaseEngine.ini for more examples
 [SectionsToSave]
 bCanSaveAllSections=false
-+Section=CurrentIniVersion
+Section=CurrentIniVersion
