Configure zenserver to use limited-lifetime mode and loopback-only mode for all platforms by default in preparation for enabling in the base engine config.

diff --git a/Engine/Config/BaseEngine.ini b/Engine/Config/BaseEngine.ini
--- a/Engine/Config/BaseEngine.ini
+++ b/Engine/Config/BaseEngine.ini
@@ -2387,10 +2387,11 @@ AutoLaunch=true
 
 [Zen.AutoLaunch]
 DataPath=%APPSETTINGSDIR%Zen/Data
+LimitProcessLifetime=true
 LocalDataCachePathEnvOverride=UE-LocalDataCachePath
 LocalDataCachePathCommandLineOverride=LocalDataCachePath
 LocalDataCachePathEditorOverrideSetting=LocalDerivedDataCache
-ExtraArgs=--http asio --gc-cache-duration-seconds 1209600 --gc-interval-seconds 21600 --gc-low-diskspace-threshold 2147483648 --quiet
+ExtraArgs=--http asio --gc-cache-duration-seconds 1209600 --gc-interval-seconds 21600 --gc-low-diskspace-threshold 2147483648 --http-forceloopback --quiet
 
 [Zen.ConnectExisting]
 HostName=localhost
