Add an unreal insights preset to analyze load times that can be used with -trace=loading on the command line

diff --git a/Engine/Config/BaseEngine.ini b/Engine/Config/BaseEngine.ini
--- a/Engine/Config/BaseEngine.ini
+++ b/Engine/Config/BaseEngine.ini
@@ -2736,6 +2736,7 @@ DrawColor=(R=140,G=255,B=0,A=164)
 
 [Trace.ChannelPresets]
 Rendering=gpu,cpu,frame,log,bookmark
+Loading=cpu,frame,log,bookmark,screenshot,region,loadtime,assetloadtime,file
 
 [RemoteConfiguration]
 Enabled=false
