Equalize packaging compression levels for DebugDev and TestShipping by default.

- Packaged data can be compressed with different codecs based on which executables are being staged, but the level it was compressed with isn't easily discoverable.
- This can add confusion to workflows where the binaries are produced separately and/or added to staged build at a later stage.
- To avoid confusion, set default compression levels across executable configurations to the same.

diff --git a/Engine/Config/BaseGame.ini b/Engine/Config/BaseGame.ini
--- a/Engine/Config/BaseGame.ini
+++ b/Engine/Config/BaseGame.ini
@@ -95,7 +95,7 @@ PackageCompressionFormat=Oodle
 ;PackageAdditionalCompressionOptions=-compressionblocksize=256KB
 PackageCompressionMethod=Kraken
 PackageCompressionLevel_Distribution=7
-PackageCompressionLevel_TestShipping=5
+PackageCompressionLevel_TestShipping=4
 PackageCompressionLevel_DebugDevelopment=4
 PackageCompressionMinBytesSaved=1024
 PackageCompressionMinPercentSaved=5
