Removes r.TSR.Velocity.Extrapolation from AA scalability

diff --git a/Engine/Config/BaseScalability.ini b/Engine/Config/BaseScalability.ini
--- a/Engine/Config/BaseScalability.ini
+++ b/Engine/Config/BaseScalability.ini
@@ -56,7 +56,6 @@ r.TSR.History.R11G11B10=1
 r.TSR.History.ScreenPercentage=100
 r.TSR.History.UpdateQuality=0
 r.TSR.ShadingRejection.Flickering=0
-r.TSR.Velocity.Extrapolation=0
 r.TSR.RejectionAntiAliasingQuality=0
 r.TSR.Resurrection=0
 
@@ -67,7 +66,6 @@ r.TSR.History.R11G11B10=1
 r.TSR.History.ScreenPercentage=100
 r.TSR.History.UpdateQuality=1
 r.TSR.ShadingRejection.Flickering=0
-r.TSR.Velocity.Extrapolation=0
 r.TSR.RejectionAntiAliasingQuality=1
 r.TSR.Resurrection=1
 
@@ -78,7 +76,6 @@ r.TSR.History.R11G11B10=1
 r.TSR.History.ScreenPercentage=100
 r.TSR.History.UpdateQuality=2
 r.TSR.ShadingRejection.Flickering=1
-r.TSR.Velocity.Extrapolation=0
 r.TSR.RejectionAntiAliasingQuality=1
 r.TSR.Resurrection=1
 
@@ -89,7 +86,6 @@ r.TSR.History.R11G11B10=1
 r.TSR.History.ScreenPercentage=200
 r.TSR.History.UpdateQuality=3
 r.TSR.ShadingRejection.Flickering=1
-r.TSR.Velocity.Extrapolation=1
 r.TSR.RejectionAntiAliasingQuality=2
 r.TSR.Resurrection=1
 
@@ -100,7 +96,6 @@ r.TSR.History.R11G11B10=0
 r.TSR.History.ScreenPercentage=200
 r.TSR.History.UpdateQuality=3
 r.TSR.ShadingRejection.Flickering=1
-r.TSR.Velocity.Extrapolation=1
 r.TSR.RejectionAntiAliasingQuality=2
 r.TSR.Resurrection=1
 
