Add ZenLocal and ZenShared to all DDC graphs in BaseEngine.ini based on presence of existing Local/Shared nodes.  ZenLocal is active and replaces Local.  ZenLocal is serviced on all desktop platforms by an autolaunched instance of zensever process on the user's machine that operates in limited lifetime mode and only listens on the loopback interface. Local is left in DeleteOnly mode for an 8 day cleanup window (down from 34 day cleanup in the past).  ZenShared is not active by default (Host=None), and requires either environment variable or commandline supplied host to utilize and remains experimental.

diff --git a/Engine/Config/BaseEngine.ini b/Engine/Config/BaseEngine.ini
--- a/Engine/Config/BaseEngine.ini
+++ b/Engine/Config/BaseEngine.ini
@@ -2454,7 +2454,7 @@ Shared=(Host=None, Namespace="ue.ddc", EnvHostOverride=UE-ZenSharedDataCacheHost
 
 [DerivedDataBackendGraph]
 ; Configure the hierarchy as the root of the cache. Reads access the nodes in order until satisfied. Writes go to every writable node.
-Root=(Type=Hierarchical, Inner=Pak, Inner=EnginePak, Inner=Local, Inner=Shared, Inner=Cloud)
+Root=(Type=Hierarchical, Inner=Pak, Inner=EnginePak, Inner=ZenLocal, Inner=Local, Inner=ZenShared, Inner=Shared, Inner=Cloud)
 ; Configure the local cache. This is a FileSystem cache with the following parameters:
 ;	ReadOnly: Can data be written to this layer. Default is false.
 ;	Touch: Update timestamps on every access instead of only after 1/4 of UnusedFileAge. Default is false.
@@ -2477,7 +2477,7 @@ Root=(Type=Hierarchical, Inner=Pak, Inner=EnginePak, Inner=Local, Inner=Shared,
 ;		ValueOnly: Look only for compressed data in the 5.0+ directory layout.
 ;		LegacyOnly: Look only for uncompressed data in the pre-5.0 directory layout.
 ;		ValueWithLegacyFallback: Look in the 5.0+ location but fall back to the pre-5.0 location on miss.
-Local=(Type=FileSystem, UnusedFileAge=34, PromptIfMissing=true, Path=%ENGINEDIR%DerivedDataCache, EnvPathOverride=UE-LocalDataCachePath, EditorOverrideSetting=LocalDerivedDataCache, CommandLineOverride=LocalDataCachePath)
+Local=(Type=FileSystem, DeleteOnly=true, UnusedFileAge=8, Path=%ENGINEDIR%DerivedDataCache, EnvPathOverride=UE-LocalDataCachePath, EditorOverrideSetting=LocalDerivedDataCache, CommandLineOverride=LocalDataCachePath)
 ; Configure the shared cache that is accessed after local.
 Shared=(Type=FileSystem, UnusedFileAge=10, FoldersToClean=10, ConsiderSlowAt=70, Path=?EpicDDC, EnvPathOverride=UE-SharedDataCachePath, EditorOverrideSetting=SharedDerivedDataCache, CommandLineOverride=SharedDataCachePath)
 ; Configure the cloud cache that is accessed after shared.
@@ -2487,18 +2487,22 @@ Cloud=(Type=Cloud, ServerID=Cloud, EnvHostOverride=UE-CloudDataCacheHost, Comman
 Pak=(Type=ReadPak, Filename="%GAMEDIR%DerivedDataCache/DDC.ddp")
 ; Configure a Project Pak node. This is a pre-generated DDC data file for the engine that can be distributed to reduce runtime fetches/generation
 EnginePak=(Type=ReadPak, Filename="%ENGINEDIR%DerivedDataCache/DDC.ddp")
+ZenLocal=(Type=Zen, ServerID=Local)
+ZenShared=(Type=Zen, ServerID=Shared)
 
 [InstalledDerivedDataBackendGraph]
 Root=(Type=KeyLength, Length=120, Inner=AsyncPut)
 AsyncPut=(Type=AsyncPut, Inner=Hierarchy)
-Hierarchy=(Type=Hierarchical, Inner=Pak, Inner=CompressedPak, Inner=EnginePak, Inner=EnterprisePak, Inner=Local, Inner=Shared, Inner=Cloud)
-Local=(Type=FileSystem, UnusedFileAge=34, Path="%ENGINEVERSIONAGNOSTICUSERDIR%DerivedDataCache", EnvPathOverride=UE-LocalDataCachePath, EditorOverrideSetting=LocalDerivedDataCache)
+Hierarchy=(Type=Hierarchical, Inner=Pak, Inner=CompressedPak, Inner=EnginePak, Inner=EnterprisePak, Inner=ZenLocal, Inner=Local, Inner=ZenShared, Inner=Shared, Inner=Cloud)
+Local=(Type=FileSystem, DeleteOnly=true, UnusedFileAge=8, Path="%ENGINEVERSIONAGNOSTICUSERDIR%DerivedDataCache", EnvPathOverride=UE-LocalDataCachePath, EditorOverrideSetting=LocalDerivedDataCache)
 Shared=(Type=FileSystem, UnusedFileAge=10, FoldersToClean=10, ConsiderSlowAt=70, Path=?EpicDDC, EnvPathOverride=UE-SharedDataCachePath, EditorOverrideSetting=SharedDerivedDataCache, CommandLineOverride=SharedDataCachePath)
 Cloud=(Type=Cloud, ServerID=Cloud, EnvHostOverride=UE-CloudDataCacheHost, CommandLineHostOverride=CloudDataCacheHost, CommandLineOAuthSecretOverride="CloudDataCacheOAuthSecret", OAuthAccessTokenEnvOverride="UE-CloudDataCacheAccessToken", EnvHttpVersionOverride="UE-CloudDataCacheHttpVersion", CommandLineHttpVersionOverride="CloudDataCacheHttpVersion")
 Pak=(Type=ReadPak, Filename="%GAMEDIR%DerivedDataCache/DDC.ddp")
 CompressedPak=(Type=ReadPak, Filename="%GAMEDIR%DerivedDataCache/Compressed.ddp", Compressed=true)
 EnginePak=(Type=ReadPak, Filename=../../../Engine/DerivedDataCache/Compressed.ddp, Compressed=true)
 EnterprisePak=(Type=ReadPak, Filename=../../../Enterprise/DerivedDataCache/Compressed.ddp, Compressed=true)
+ZenLocal=(Type=Zen, ServerID=Local)
+ZenShared=(Type=Zen, ServerID=Shared)
 
 [ZenOnlyDDC]
 Root=(Type=KeyLength, Length=120, Inner=AsyncPut)
@@ -2506,9 +2510,10 @@ AsyncPut=(Type=AsyncPut, Inner=Zen)
 Zen=(Type=Zen, ServerID=Local)
 
 [ZenDDC]
-Root=(Type=Hierarchical, Inner=Pak, Inner=EnginePak, Inner=Local, Inner=ZenLocal, Inner=Shared, Inner=Cloud)
-Local=(Type=FileSystem, DeleteOnly=true, UnusedFileAge=34, PromptIfMissing=true, Path=%ENGINEDIR%DerivedDataCache, EnvPathOverride=UE-LocalDataCachePath, EditorOverrideSetting=LocalDerivedDataCache, CommandLineOverride=LocalDataCachePath)
+Root=(Type=Hierarchical, Inner=Pak, Inner=EnginePak, Inner=ZenLocal, Inner=Local, Inner=ZenShared, Inner=Shared, Inner=Cloud)
+Local=(Type=FileSystem, DeleteOnly=true, UnusedFileAge=8, Path=%ENGINEDIR%DerivedDataCache, EnvPathOverride=UE-LocalDataCachePath, EditorOverrideSetting=LocalDerivedDataCache, CommandLineOverride=LocalDataCachePath)
 ZenLocal=(Type=Zen, ServerID=Local)
+ZenShared=(Type=Zen, ServerID=Shared)
 Shared=(Type=FileSystem, UnusedFileAge=10, FoldersToClean=10, ConsiderSlowAt=70, Path=?EpicDDC, EnvPathOverride=UE-SharedDataCachePath, EditorOverrideSetting=SharedDerivedDataCache, CommandLineOverride=SharedDataCachePath)
 Cloud=(Type=Cloud, ServerID=Cloud, EnvHostOverride=UE-CloudDataCacheHost, CommandLineHostOverride=CloudDataCacheHost, CommandLineOAuthSecretOverride="CloudDataCacheOAuthSecret", OAuthAccessTokenEnvOverride="UE-CloudDataCacheAccessToken", EnvHttpVersionOverride="UE-CloudDataCacheHttpVersion", CommandLineHttpVersionOverride="CloudDataCacheHttpVersion")
 Pak=(Type=ReadPak, Filename="%GAMEDIR%DerivedDataCache/DDC.ddp")
@@ -2529,80 +2534,92 @@ ZenLocal=(Type=Zen, ServerID=Local, Sandbox=Cold)
 [DerivedDataBackendGraph_Fill_Seattle]
 Root=(Type=KeyLength, Length=120, Inner=AsyncPut)
 AsyncPut=(Type=AsyncPut, Inner=Hierarchy)
-Hierarchy=(Type=Hierarchical, Inner=Pak, Inner=EnginePak, Inner=Local, Inner=Seattle)
-Local=(Type=FileSystem, UnusedFileAge=34, Path=%ENGINEDIR%DerivedDataCache, EnvPathOverride=UE-LocalDataCachePath, EditorOverrideSetting=LocalDerivedDataCache)
+Hierarchy=(Type=Hierarchical, Inner=Pak, Inner=EnginePak, Inner=ZenLocal, Inner=Local, Inner=Seattle)
+Local=(Type=FileSystem, DeleteOnly=true, UnusedFileAge=8, Path=%ENGINEDIR%DerivedDataCache, EnvPathOverride=UE-LocalDataCachePath, EditorOverrideSetting=LocalDerivedDataCache, CommandLineOverride=LocalDataCachePath)
 Seattle=(Type=FileSystem, UnusedFileAge=23, FoldersToClean=10, Path=?EpicSeaDDC, EnvPathOverride=UE-SharedDataCachePath_Seattle)
 Pak=(Type=ReadPak, Filename="%GAMEDIR%DerivedDataCache/DDC.ddp")
 EnginePak=(Type=ReadPak, Filename="%ENGINEDIR%DerivedDataCache/DDC.ddp")
+ZenLocal=(Type=Zen, ServerID=Local)
 
 [NoShared]
 Root=(Type=KeyLength, Length=120, Inner=AsyncPut)
 AsyncPut=(Type=AsyncPut, Inner=Hierarchy)
-Hierarchy=(Type=Hierarchical, Inner=Pak, Inner=Local)
-Local=(Type=FileSystem, UnusedFileAge=34, Path=%ENGINEDIR%DerivedDataCache, EnvPathOverride=UE-LocalDataCachePath, EditorOverrideSetting=LocalDerivedDataCache, CommandLineOverride=LocalDataCachePath)
+Hierarchy=(Type=Hierarchical, Inner=Pak, Inner=ZenLocal, Inner=Local)
+Local=(Type=FileSystem, DeleteOnly=true, UnusedFileAge=8, Path=%ENGINEDIR%DerivedDataCache, EnvPathOverride=UE-LocalDataCachePath, EditorOverrideSetting=LocalDerivedDataCache, CommandLineOverride=LocalDataCachePath)
 Pak=(Type=ReadPak, Filename="%GAMEDIR%DerivedDataCache/DDC.ddp")
+ZenLocal=(Type=Zen, ServerID=Local)
 
 ; Creates a completely new standalone cache so as to not destroy a valuable populated local cache
 [Cold]
 Root=(Type=KeyLength, Length=120, Inner=AsyncPut)
 AsyncPut=(Type=AsyncPut, Inner=Hierarchy)
-Hierarchy=(Type=Hierarchical, Inner=Local)
-Local=(Type=FileSystem, Flush=true, UnusedFileAge=1, Path="%GAMEDIR%ColdDerivedDataCache")
+Hierarchy=(Type=Hierarchical, Inner=ZenLocal)
+ZenLocal=(Type=Zen, ServerID=Local, Sandbox=Cold, Flush=true)
 
 ; Inherits the same standalone cache from [Cold], for incremental performance measurement purposes
 [Warm]
 Root=(Type=KeyLength, Length=120, Inner=AsyncPut)
 AsyncPut=(Type=AsyncPut, Inner=Hierarchy)
-Hierarchy=(Type=Hierarchical, Inner=Local)
-Local=(Type=FileSystem, Clean=true, UnusedFileAge=1, Path="%GAMEDIR%ColdDerivedDataCache")
+Hierarchy=(Type=Hierarchical, Inner=ZenLocal)
+ZenLocal=(Type=Zen, ServerID=Local, Sandbox=Cold)
 
 [CreatePak]
 Root=(Type=KeyLength, Length=120, Inner=AsyncPut)
 AsyncPut=(Type=AsyncPut, Inner=Hierarchy)
-Hierarchy=(Type=Hierarchical, Inner=PakWrite, Inner=PakRead, Inner=Local, Inner=Shared, Inner=Cloud)
-Local=(Type=FileSystem, UnusedFileAge=34, Path=%ENGINEDIR%DerivedDataCache, EnvPathOverride=UE-LocalDataCachePath, EditorOverrideSetting=LocalDerivedDataCache)
+Hierarchy=(Type=Hierarchical, Inner=PakWrite, Inner=PakRead, Inner=ZenLocal, Inner=Local, Inner=ZenShared, Inner=Shared, Inner=Cloud)
+Local=(Type=FileSystem, DeleteOnly=true, UnusedFileAge=8, Path=%ENGINEDIR%DerivedDataCache, EnvPathOverride=UE-LocalDataCachePath, EditorOverrideSetting=LocalDerivedDataCache, CommandLineOverride=LocalDataCachePath)
 Shared=(Type=FileSystem, DeleteUnused=false, ConsiderSlowAt=70, Path=?EpicDDC, EnvPathOverride=UE-SharedDataCachePath, EditorOverrideSetting=SharedDerivedDataCache, CommandLineOverride=SharedDataCachePath)
 Cloud=(Type=Cloud, ServerID=Cloud, EnvHostOverride=UE-CloudDataCacheHost, CommandLineHostOverride=CloudDataCacheHost, CommandLineOAuthSecretOverride="CloudDataCacheOAuthSecret", OAuthAccessTokenEnvOverride="UE-CloudDataCacheAccessToken", EnvHttpVersionOverride="UE-CloudDataCacheHttpVersion", CommandLineHttpVersionOverride="CloudDataCacheHttpVersion")
 PakRead=(Type=ReadPak, Filename="%GAMEDIR%DerivedDataCache/DDC.ddp")
 PakWrite=(Type=WritePak, Filename="%GAMEDIR%DerivedDataCache/DDC.ddp")
+ZenLocal=(Type=Zen, ServerID=Local)
+ZenShared=(Type=Zen, ServerID=Shared)
 
 [CreateInstalledProjectPak]
 Root=(Type=KeyLength, Length=120, Inner=AsyncPut)
 AsyncPut=(Type=AsyncPut, Inner=Hierarchy)
-Hierarchy=(Type=Hierarchical, Inner=EnginePak, Inner=PakWrite, Inner=PakRead, Inner=Local, Inner=Shared, Inner=Cloud)
-Local=(Type=FileSystem, UnusedFileAge=34, Path=%ENGINEDIR%DerivedDataCache, EnvPathOverride=UE-LocalDataCachePath, EditorOverrideSetting=LocalDerivedDataCache)
+Hierarchy=(Type=Hierarchical, Inner=EnginePak, Inner=PakWrite, Inner=PakRead, Inner=ZenLocal, Inner=Local, Inner=ZenShared, Inner=Shared, Inner=Cloud)
+Local=(Type=FileSystem, DeleteOnly=true, UnusedFileAge=8, Path=%ENGINEDIR%DerivedDataCache, EnvPathOverride=UE-LocalDataCachePath, EditorOverrideSetting=LocalDerivedDataCache, CommandLineOverride=LocalDataCachePath)
 Shared=(Type=FileSystem, DeleteUnused=false, ConsiderSlowAt=70, Path=?EpicDDC, EnvPathOverride=UE-SharedDataCachePath, EditorOverrideSetting=SharedDerivedDataCache, CommandLineOverride=SharedDataCachePath)
 Cloud=(Type=Cloud, ServerID=Cloud, EnvHostOverride=UE-CloudDataCacheHost, CommandLineHostOverride=CloudDataCacheHost, CommandLineOAuthSecretOverride="CloudDataCacheOAuthSecret", OAuthAccessTokenEnvOverride="UE-CloudDataCacheAccessToken", EnvHttpVersionOverride="UE-CloudDataCacheHttpVersion", CommandLineHttpVersionOverride="CloudDataCacheHttpVersion")
 EnginePak=(Type=ReadPak, Filename=../../../Engine/DerivedDataCache/Compressed.ddp, Compressed=true)
 PakRead=(Type=ReadPak, Filename="%GAMEDIR%DerivedDataCache/Compressed.ddp", Compressed=true)
 PakWrite=(Type=WritePak, Filename="%GAMEDIR%DerivedDataCache/Compressed.ddp", Compressed=true)
+ZenLocal=(Type=Zen, ServerID=Local)
+ZenShared=(Type=Zen, ServerID=Shared)
 
 [CreateInstalledEnginePak]
 Root=(Type=KeyLength, Length=120, Inner=AsyncPut)
 AsyncPut=(Type=AsyncPut, Inner=Hierarchy)
-Hierarchy=(Type=Hierarchical, Inner=PakWrite, Inner=Local, Inner=Shared, Inner=Cloud)
-Local=(Type=FileSystem, UnusedFileAge=34, Path=%ENGINEDIR%DerivedDataCache, EnvPathOverride=UE-LocalDataCachePath, EditorOverrideSetting=LocalDerivedDataCache)
+Hierarchy=(Type=Hierarchical, Inner=PakWrite, Inner=ZenLocal, Inner=Local, Inner=ZenShared, Inner=Shared, Inner=Cloud)
+Local=(Type=FileSystem, DeleteOnly=true, UnusedFileAge=8, Path=%ENGINEDIR%DerivedDataCache, EnvPathOverride=UE-LocalDataCachePath, EditorOverrideSetting=LocalDerivedDataCache, CommandLineOverride=LocalDataCachePath)
 Shared=(Type=FileSystem, DeleteUnused=false, ConsiderSlowAt=70, Path=?EpicDDC, EnvPathOverride=UE-SharedDataCachePath, EditorOverrideSetting=SharedDerivedDataCache, CommandLineOverride=SharedDataCachePath)
 Cloud=(Type=Cloud, ServerID=Cloud, EnvHostOverride=UE-CloudDataCacheHost, CommandLineHostOverride=CloudDataCacheHost, CommandLineOAuthSecretOverride="CloudDataCacheOAuthSecret", OAuthAccessTokenEnvOverride="UE-CloudDataCacheAccessToken", EnvHttpVersionOverride="UE-CloudDataCacheHttpVersion", CommandLineHttpVersionOverride="CloudDataCacheHttpVersion")
 PakWrite=(Type=WritePak, Filename="%ENGINEDIR%DerivedDataCache/Compressed.ddp", Compressed=true)
+ZenLocal=(Type=Zen, ServerID=Local)
+ZenShared=(Type=Zen, ServerID=Shared)
 
 [CreateInstalledEnterprisePak]
 Root=(Type=KeyLength, Length=120, Inner=AsyncPut)
 AsyncPut=(Type=AsyncPut, Inner=Hierarchy)
-Hierarchy=(Type=Hierarchical, Inner=PakWrite, Inner=Local, Inner=Shared, Inner=Cloud)
-Local=(Type=FileSystem, UnusedFileAge=34, Path=%ENGINEDIR%DerivedDataCache, EnvPathOverride=UE-LocalDataCachePath, EditorOverrideSetting=LocalDerivedDataCache)
+Hierarchy=(Type=Hierarchical, Inner=PakWrite, Inner=ZenLocal, Inner=Local, Inner=ZenShared, Inner=Shared, Inner=Cloud)
+Local=(Type=FileSystem, DeleteOnly=true, UnusedFileAge=8, Path=%ENGINEDIR%DerivedDataCache, EnvPathOverride=UE-LocalDataCachePath, EditorOverrideSetting=LocalDerivedDataCache, CommandLineOverride=LocalDataCachePath)
 Shared=(Type=FileSystem, DeleteUnused=false, ConsiderSlowAt=70, Path=?EpicDDC, EnvPathOverride=UE-SharedDataCachePath, EditorOverrideSetting=SharedDerivedDataCache, CommandLineOverride=SharedDataCachePath)
 Cloud=(Type=Cloud, ServerID=Cloud, EnvHostOverride=UE-CloudDataCacheHost, CommandLineHostOverride=CloudDataCacheHost, CommandLineOAuthSecretOverride="CloudDataCacheOAuthSecret", OAuthAccessTokenEnvOverride="UE-CloudDataCacheAccessToken", EnvHttpVersionOverride="UE-CloudDataCacheHttpVersion", CommandLineHttpVersionOverride="CloudDataCacheHttpVersion")
 PakWrite=(Type=WritePak, Filename=../../../Enterprise/DerivedDataCache/Compressed.ddp, Compressed=true)
+ZenLocal=(Type=Zen, ServerID=Local)
+ZenShared=(Type=Zen, ServerID=Shared)
 
 [CreateProjectCache]
 Root=(Type=KeyLength, Length=120, Inner=AsyncPut)
 AsyncPut=(Type=AsyncPut, Inner=Hierarchy)
-Hierarchy=(Type=Hierarchical, Inner=Local, Inner=Project, Inner=Shared, Inner=Cloud)
-Local=(Type=FileSystem, UnusedFileAge=34, Path=%ENGINEDIR%DerivedDataCache, EnvPathOverride=UE-LocalDataCachePath, EditorOverrideSetting=LocalDerivedDataCache)
+Hierarchy=(Type=Hierarchical, Inner=ZenLocal, Inner=Local, Inner=Project, Inner=ZenShared, Inner=Shared, Inner=Cloud)
+Local=(Type=FileSystem, DeleteOnly=true, UnusedFileAge=8, Path=%ENGINEDIR%DerivedDataCache, EnvPathOverride=UE-LocalDataCachePath, EditorOverrideSetting=LocalDerivedDataCache, CommandLineOverride=LocalDataCachePath)
 Project=(Type=FileSystem, Flush=true, UnusedFileAge=34, Path=%GAMEDIR%ProjectDerivedData)
 Shared=(Type=FileSystem, DeleteUnused=false, ConsiderSlowAt=70, Path=?EpicDDC, EnvPathOverride=UE-SharedDataCachePath, EditorOverrideSetting=SharedDerivedDataCache, CommandLineOverride=SharedDataCachePath)
 Cloud=(Type=Cloud, ServerID=Cloud, EnvHostOverride=UE-CloudDataCacheHost, CommandLineHostOverride=CloudDataCacheHost, CommandLineOAuthSecretOverride="CloudDataCacheOAuthSecret", OAuthAccessTokenEnvOverride="UE-CloudDataCacheAccessToken")
+ZenLocal=(Type=Zen, ServerID=Local)
+ZenShared=(Type=Zen, ServerID=Shared)
 
 [VirtualTextureChunkDDCCache]
 UnusedFileAge=34
