Interchange project setting do not save to local engine ini file anymore, adding the interchange section to the [SectionsToSave] of BaseEngine.ini allow it to be save again.
#rnx

diff --git a/Engine/Config/BaseEngine.ini b/Engine/Config/BaseEngine.ini
--- a/Engine/Config/BaseEngine.ini
+++ b/Engine/Config/BaseEngine.ini
@@ -4,6 +4,7 @@
 
 [SectionsToSave]
 +Section=GameNetDriver StatelessConnectHandlerComponent
++Section=/Script/InterchangeEngine.InterchangeProjectSettings
 
 [/Script/EngineSettings.GameMapsSettings]
 GameInstanceClass=/Script/Engine.GameInstance
