Android emulator support
- Device profile to force Vulkan in the emulator
- Fall back to ES2 context if necessary for device detection

diff --git a/Engine/Config/BaseDeviceProfiles.ini b/Engine/Config/BaseDeviceProfiles.ini
--- a/Engine/Config/BaseDeviceProfiles.ini
+++ b/Engine/Config/BaseDeviceProfiles.ini
@@ -700,6 +700,9 @@ BaseProfileName=VisionPro
 +MatchProfile=(Profile="Oculus_Quest",Match=((SourceType=SRC_DeviceMake,CompareType=CMP_Equal,MatchString="Oculus")))
 +MatchProfile=(Profile="MagicLeap_Vulkan",Match=((SourceType=SRC_DeviceMake,CompareType=CMP_Equal,MatchString="MagicLeap"), (SourceType=SRC_VulkanAvailable,CompareType=CMP_Equal,MatchString="true")))
 
+; Match Android emulator
++MatchProfile=(Profile="Android_Emulator",Match=((SourceType=SRC_DeviceMake,CompareType=CMP_Equal,MatchString="Google"),(SourceType=SRC_GPUFamily,CompareType=CMP_Regex,MatchString="Android Emulator")))
+
 ; Adreno-based devices
 +MatchProfile=(Profile="Android_Adreno4xx",Match=((SourceType=SRC_GpuFamily,CompareType=CMP_Regex,MatchString="Adreno \\(TM\\) 4[0-9][0-9]")))
 +MatchProfile=(Profile="Android_Adreno5xx_Low",Match=((SourceType=SRC_GpuFamily,CompareType=CMP_Regex,MatchString="Adreno \\(TM\\) 5[0-1][0-9]")))
@@ -840,6 +843,11 @@ BaseProfileName=Android
 +CVars=sg.LandscapeQuality=2
 +CVars=r.MobileContentScaleFactor=1.0
 
+[Android_Emulator DeviceProfile]
+DeviceType=Android
+BaseProfileName=Android_High
++CVars=r.Android.DisableVulkanSupport=0
+
 ; Example base settings for Vulkan SM5 on Android
 [Android_Vulkan_SM5 DeviceProfile]
 DeviceType=Android
diff --git a/Engine/Source/Runtime/OpenGLDrv/Private/Android/AndroidEGL.cpp b/Engine/Source/Runtime/OpenGLDrv/Private/Android/AndroidEGL.cpp
--- a/Engine/Source/Runtime/OpenGLDrv/Private/Android/AndroidEGL.cpp
+++ b/Engine/Source/Runtime/OpenGLDrv/Private/Android/AndroidEGL.cpp
@@ -824,7 +824,18 @@ void AndroidEGL::Init(APIVariant API, uint32 MajorVersion, uint32 MinorVersion)
 	{
 		ContextAttributes[3] -= 1;
 
-		InitContexts();
+		bSuccess = InitContexts();
+
+		if (!bSuccess)
+		{
+			// Try to create an ES2 context if ES3.1 also failed, which can happen in the Android emulator.
+			// This is enough for FAndroidGPUInfo detection to enable Vulkan.
+			ContextAttributes[0] = EGL_CONTEXT_CLIENT_VERSION;
+			ContextAttributes[1] = 2;
+			ContextAttributes[2] = EGL_NONE;
+
+			bSuccess = InitContexts();
+		}
 	}
 
 	// Getting the hardware window is valid during preinit as we have GAndroidWindowLock held.
