Minor fix to device profile submission from CL 30526383

diff --git a/Engine/Config/BaseDeviceProfiles.ini b/Engine/Config/BaseDeviceProfiles.ini
--- a/Engine/Config/BaseDeviceProfiles.ini
+++ b/Engine/Config/BaseDeviceProfiles.ini
@@ -93,6 +93,7 @@
 +DeviceProfileNameAndTypes=Android_Low,Android
 +DeviceProfileNameAndTypes=Android_Mid,Android
 +DeviceProfileNameAndTypes=Android_High,Android
++DeviceProfileNameAndTypes=Android_Emulator,Android
 +DeviceProfileNameAndTypes=Android_Default,Android
 +DeviceProfileNameAndTypes=Android_Adreno4xx,Android
 +DeviceProfileNameAndTypes=Android_Adreno5xx_Low,Android
