Add SRC_SM5Available similar to SRC_VulkanAvailable that can be used in device profile ini to match based in Desktop renderer support.

New Mali Immortalis and Adreno 7xx rules now enable desktop renderer by default if it is built into the APK (that part is still off by default).

diff --git a/Engine/Config/BaseDeviceProfiles.ini b/Engine/Config/BaseDeviceProfiles.ini
--- a/Engine/Config/BaseDeviceProfiles.ini
+++ b/Engine/Config/BaseDeviceProfiles.ini
@@ -709,10 +709,8 @@ BaseProfileName=VisionPro
 +MatchProfile=(Profile="Android_Adreno5xx_Low",Match=((SourceType=SRC_GpuFamily,CompareType=CMP_Regex,MatchString="Adreno \\(TM\\) 5[0-1][0-9]")))
 +MatchProfile=(Profile="Android_Adreno5xx",Match=((SourceType=SRC_GpuFamily,CompareType=CMP_Regex,MatchString="Adreno \\(TM\\) 5[0-9][0-9]")))
 ; Example profiles for Vulkan SM5
-;+MatchProfile=(Profile="Android_Vulkan_SM5",Match=((SourceType=SRC_GpuFamily,CompareType=CMP_Regex,MatchString="Adreno \\(TM\\) 6[0-9][0-9]"),(SourceType=SRC_AndroidVersion, CompareType=CMP_Regex,MatchString="([0-9]+).*"),(SourceType=SRC_PreviousRegexMatch,CompareType=CMP_GreaterEqual,MatchString="9"),(SourceType=SRC_VulkanAvailable,CompareType=CMP_Equal,MatchString="true")))
-;+MatchProfile=(Profile="Android_Vulkan_SM5",Match=((SourceType=SRC_GpuFamily,CompareType=CMP_Regex,MatchString="Adreno \\(TM\\) 7[0-9][0-9]"),(SourceType=SRC_AndroidVersion, CompareType=CMP_Regex,MatchString="([0-9]+).*"),(SourceType=SRC_PreviousRegexMatch,CompareType=CMP_GreaterEqual,MatchString="9"),(SourceType=SRC_VulkanAvailable,CompareType=CMP_Equal,MatchString="true")))
-;+MatchProfile=(Profile="Android_Vulkan_SM5_Mali",Match=((SourceType=SRC_GpuFamily,CompareType=CMP_Regex,MatchString="^Mali\\-G7[67]"),(SourceType=SRC_VulkanAvailable,CompareType=CMP_Equal,MatchString="true")))
-;+MatchProfile=(Profile="Android_Vulkan_SM5_Mali",Match=((SourceType=SRC_GpuFamily,CompareType=CMP_Regex,MatchString="^Mali\\-G715"),(SourceType=SRC_VulkanAvailable,CompareType=CMP_Equal,MatchString="true")))
+;+MatchProfile=(Profile="Android_Vulkan_SM5",Match=((SourceType=SRC_GpuFamily,CompareType=CMP_Regex,MatchString="Adreno \\(TM\\) 6[0-9][0-9]"),(SourceType=SRC_AndroidVersion, CompareType=CMP_Regex,MatchString="([0-9]+).*"),(SourceType=SRC_PreviousRegexMatch,CompareType=CMP_GreaterEqual,MatchString="9"),(SourceType=SRC_SM5Available,CompareType=CMP_Equal,MatchString="true")))
++MatchProfile=(Profile="Android_Vulkan_SM5",Match=((SourceType=SRC_GpuFamily,CompareType=CMP_Regex,MatchString="Adreno \\(TM\\) 7[0-9][0-9]"),(SourceType=SRC_AndroidVersion, CompareType=CMP_Regex,MatchString="([0-9]+).*"),(SourceType=SRC_PreviousRegexMatch,CompareType=CMP_GreaterEqual,MatchString="9"),(SourceType=SRC_SM5Available,CompareType=CMP_Equal,MatchString="true")))
 ; disable Vulkan on Adreno 660 devices with driver version <= 525
 +MatchProfile=(Profile="Android_Adreno6xx",Match=((SourceType=SRC_GpuFamily,CompareType=CMP_Regex,MatchString="Adreno \\(TM\\) 660"),(SourceType=SRC_GLVersion, CompareType=CMP_Regex,MatchString="V@([0-9]+)"),(SourceType=SRC_PreviousRegexMatch,CompareType=CMP_LessEqual,MatchString="525")))
 ; Enable Vulkan on Adreno 6xx devices with Android >= 10 and driver version >= 400
@@ -740,6 +738,9 @@ BaseProfileName=VisionPro
 +MatchProfile=(Profile="Android_Mali_G710",Match=((SourceType=SRC_GpuFamily,CompareType=CMP_Regex,MatchString="^Mali\\-G710")))
 +MatchProfile=(Profile="Android_Mali_G7xx_Vulkan",Match=((SourceType=SRC_GpuFamily,CompareType=CMP_Regex,MatchString="^Mali\\-G7[0-9][0-9]"),(SourceType=SRC_VulkanAvailable,CompareType=CMP_Equal,MatchString="true")))
 +MatchProfile=(Profile="Android_Mali_G7xx",Match=((SourceType=SRC_GpuFamily,CompareType=CMP_Regex,MatchString="^Mali\\-G7[0-9][0-9]")))
+; Example profiles for Vulkan SM5
+;+MatchProfile=(Profile="Android_Vulkan_SM5_Mali",Match=((SourceType=SRC_GpuFamily,CompareType=CMP_Regex,MatchString="^Mali\\-G7[67]"),(SourceType=SRC_SM5Available,CompareType=CMP_Equal,MatchString="true")))
++MatchProfile=(Profile="Android_Vulkan_SM5_Mali",Match=((SourceType=SRC_GpuFamily,CompareType=CMP_Regex,MatchString="^Mali\\-G7[1-9][0-9]"),(SourceType=SRC_SM5Available,CompareType=CMP_Equal,MatchString="true")))
 ; Enable Vulkan on any unknown Mali running Android 10 or later
 +MatchProfile=(Profile="Android_Unknown_Vulkan",Match=((SourceType=SRC_GpuFamily,CompareType=CMP_Regex,MatchString="^Mali"),(SourceType=SRC_AndroidVersion, CompareType=CMP_Regex,MatchString="([0-9]+).*"),(SourceType=SRC_PreviousRegexMatch,CompareType=CMP_GreaterEqual,MatchString="10"),(SourceType=SRC_VulkanAvailable,CompareType=CMP_Equal,MatchString="true")))
 
diff --git a/Engine/Plugins/Runtime/AndroidDeviceProfileSelector/Source/AndroidDeviceProfileSelector/Private/AndroidDeviceProfileMatchingRules.h b/Engine/Plugins/Runtime/AndroidDeviceProfileSelector/Source/AndroidDeviceProfileSelector/Private/AndroidDeviceProfileMatchingRules.h
--- a/Engine/Plugins/Runtime/AndroidDeviceProfileSelector/Source/AndroidDeviceProfileSelector/Private/AndroidDeviceProfileMatchingRules.h
+++ b/Engine/Plugins/Runtime/AndroidDeviceProfileSelector/Source/AndroidDeviceProfileSelector/Private/AndroidDeviceProfileMatchingRules.h
@@ -22,6 +22,7 @@ enum ESourceType : int
 	SRC_Chipset,
 	SRC_ConfigRuleVar,
 	SRC_HMDSystemName,
+	SRC_SM5Available,
 	SRC_MAX,
 };
 
diff --git a/Engine/Plugins/Runtime/AndroidDeviceProfileSelector/Source/AndroidDeviceProfileSelector/Private/AndroidDeviceProfileSelector.cpp b/Engine/Plugins/Runtime/AndroidDeviceProfileSelector/Source/AndroidDeviceProfileSelector/Private/AndroidDeviceProfileSelector.cpp
--- a/Engine/Plugins/Runtime/AndroidDeviceProfileSelector/Source/AndroidDeviceProfileSelector/Private/AndroidDeviceProfileSelector.cpp
+++ b/Engine/Plugins/Runtime/AndroidDeviceProfileSelector/Source/AndroidDeviceProfileSelector/Private/AndroidDeviceProfileSelector.cpp
@@ -50,6 +50,7 @@ void FAndroidDeviceProfileSelector::VerifySelectorParams()
 	check(SelectorProperties.Find(FAndroidProfileSelectorSourceProperties::SRC_Chipset));
 	check(SelectorProperties.Find(FAndroidProfileSelectorSourceProperties::SRC_HMDSystemName));
 	check(SelectorProperties.Find(FAndroidProfileSelectorSourceProperties::SRC_TotalPhysicalGB));
+	check(SelectorProperties.Find(FAndroidProfileSelectorSourceProperties::SRC_SM5Available));
 }
 
 FString FAndroidDeviceProfileSelector::FindMatchingProfile(const FString& FallbackProfileName)
@@ -70,6 +71,7 @@ FString FAndroidDeviceProfileSelector::FindMatchingProfile(const FString& Fallba
 	FString Chipset				= SelectorProperties.FindChecked(FAndroidProfileSelectorSourceProperties::SRC_Chipset);
 	FString TotalPhysicalGB		= SelectorProperties.FindChecked(FAndroidProfileSelectorSourceProperties::SRC_TotalPhysicalGB);
 	FString HMDSystemName		= SelectorProperties.FindChecked(FAndroidProfileSelectorSourceProperties::SRC_HMDSystemName);
+	FString SM5Available		= SelectorProperties.FindChecked(FAndroidProfileSelectorSourceProperties::SRC_SM5Available);
 
 	for (const FProfileMatch& Profile : GetAndroidDeviceProfileMatchingRules()->MatchProfile)
 	{
@@ -124,6 +126,9 @@ FString FAndroidDeviceProfileSelector::FindMatchingProfile(const FString& Fallba
 			case ESourceType::SRC_HMDSystemName:
 				SourceString = &HMDSystemName;
 				break;
+			case ESourceType::SRC_SM5Available:
+				SourceString = &SM5Available;
+				break;
 			case ESourceType::SRC_ConfigRuleVar:
 				{
 					// expected Matchstring contents for configrulevar is "configrule_varname|matchstring"
diff --git a/Engine/Plugins/Runtime/AndroidDeviceProfileSelector/Source/AndroidDeviceProfileSelector/Public/AndroidDeviceProfileSelector.h b/Engine/Plugins/Runtime/AndroidDeviceProfileSelector/Source/AndroidDeviceProfileSelector/Public/AndroidDeviceProfileSelector.h
--- a/Engine/Plugins/Runtime/AndroidDeviceProfileSelector/Source/AndroidDeviceProfileSelector/Public/AndroidDeviceProfileSelector.h
+++ b/Engine/Plugins/Runtime/AndroidDeviceProfileSelector/Source/AndroidDeviceProfileSelector/Public/AndroidDeviceProfileSelector.h
@@ -20,6 +20,7 @@ namespace FAndroidProfileSelectorSourceProperties
 	static FName SRC_Chipset(TEXT("SRC_Chipset"));
 	static FName SRC_HMDSystemName(TEXT("SRC_HMDSystemName"));
 	static FName SRC_TotalPhysicalGB(TEXT("SRC_TotalPhysicalGB"));
+	static FName SRC_SM5Available(TEXT("SRC_SM5Available"));
 };
 
 class ANDROIDDEVICEPROFILESELECTOR_API FAndroidDeviceProfileSelector
diff --git a/Engine/Plugins/Runtime/AndroidDeviceProfileSelector/Source/AndroidDeviceProfileSelectorRuntime/Private/AndroidDeviceProfileSelectorRuntimeModule.cpp b/Engine/Plugins/Runtime/AndroidDeviceProfileSelector/Source/AndroidDeviceProfileSelectorRuntime/Private/AndroidDeviceProfileSelectorRuntimeModule.cpp
--- a/Engine/Plugins/Runtime/AndroidDeviceProfileSelector/Source/AndroidDeviceProfileSelectorRuntime/Private/AndroidDeviceProfileSelectorRuntimeModule.cpp
+++ b/Engine/Plugins/Runtime/AndroidDeviceProfileSelector/Source/AndroidDeviceProfileSelectorRuntime/Private/AndroidDeviceProfileSelectorRuntimeModule.cpp
@@ -77,6 +77,7 @@ static const TMap<FName,FString>& GetDeviceSelectorParams()
 		AndroidParams.Add(FAndroidProfileSelectorSourceProperties::SRC_Chipset, GetParam(FString(TEXT("unknown")), TEXT("chipset")));
 		AndroidParams.Add(FAndroidProfileSelectorSourceProperties::SRC_HMDSystemName, HMDRequestedProfileName);
 		AndroidParams.Add(FAndroidProfileSelectorSourceProperties::SRC_TotalPhysicalGB, FString::Printf(TEXT("%d"), TotalPhysicalGB));
+		AndroidParams.Add(FAndroidProfileSelectorSourceProperties::SRC_SM5Available, FAndroidMisc::IsDesktopVulkanAvailable() ? TEXT("true") : TEXT("false"));
 
 #if PLATFORM_ANDROID
 		// allow ConfigRules to override cvars first
diff --git a/Engine/Source/Editor/DeviceProfileServices/Private/CheckAndroidDeviceProfileCommandlet.cpp b/Engine/Source/Editor/DeviceProfileServices/Private/CheckAndroidDeviceProfileCommandlet.cpp
--- a/Engine/Source/Editor/DeviceProfileServices/Private/CheckAndroidDeviceProfileCommandlet.cpp
+++ b/Engine/Source/Editor/DeviceProfileServices/Private/CheckAndroidDeviceProfileCommandlet.cpp
@@ -149,6 +149,7 @@ int32 UCheckAndroidDeviceProfileCommandlet::Main(const FString& RawCommandLine)
 					DeviceParameters.Add(FName(TEXT("SRC_Chipset")), AndroidProperties.Chipset);
 					DeviceParameters.Add(FName(TEXT("SRC_TotalPhysicalGB")), AndroidProperties.TotalPhysicalGB);
 					DeviceParameters.Add(FName(TEXT("SRC_HMDSystemName")), TEXT(""));
+					DeviceParameters.Add(FName(TEXT("SRC_SM5Available")), AndroidProperties.SM5Available ? "true" : "false");
 				}
 			}
 			else
@@ -172,6 +173,7 @@ int32 UCheckAndroidDeviceProfileCommandlet::Main(const FString& RawCommandLine)
 				DeviceParameters.Add(FName(TEXT("SRC_Chipset")), Params.FindRef(TEXT("Chipset")));
 				DeviceParameters.Add(FName(TEXT("SRC_TotalPhysicalGB")), Params.FindRef(TEXT("TotalPhysicalGB")));
 				DeviceParameters.Add(FName(TEXT("SRC_HMDSystemName")), TEXT(""));
+				DeviceParameters.Add(FName(TEXT("SRC_SM5Available")), Params.FindRef(TEXT("SM5Available")));
 			}
 			AndroidDeviceProfileSelector->SetSelectorProperties(DeviceParameters);
 			FString ProfileName = ForcedDeviceProfileName.IsEmpty() ? AndroidDeviceProfileSelector->GetDeviceProfileName() : ForcedDeviceProfileName;
diff --git a/Engine/Source/Editor/PIEPreviewDeviceProfileSelector/Private/PIEPreviewDevice.cpp b/Engine/Source/Editor/PIEPreviewDeviceProfileSelector/Private/PIEPreviewDevice.cpp
--- a/Engine/Source/Editor/PIEPreviewDeviceProfileSelector/Private/PIEPreviewDevice.cpp
+++ b/Engine/Source/Editor/PIEPreviewDeviceProfileSelector/Private/PIEPreviewDevice.cpp
@@ -457,6 +457,7 @@ FString FPIEPreviewDevice::GetProfile() const
 				DeviceParameters.Add(FName(TEXT("SRC_Chipset")), AndroidProperties.Chipset);
 				DeviceParameters.Add(FName(TEXT("SRC_HMDSystemName")), AndroidProperties.HMDSystemName);
 				DeviceParameters.Add(FName(TEXT("SRC_TotalPhysicalGB")), AndroidProperties.TotalPhysicalGB);
+				DeviceParameters.Add(FName(TEXT("SRC_SM5Available")), AndroidProperties.SM5Available ? "true" : "false");
 
 				AndroidDeviceProfileSelector->SetSelectorProperties(DeviceParameters);
 				FString PIEProfileName = AndroidDeviceProfileSelector->GetDeviceProfileName();
diff --git a/Engine/Source/Editor/PIEPreviewDeviceSpecification/Public/PIEPreviewDeviceSpecification.h b/Engine/Source/Editor/PIEPreviewDeviceSpecification/Public/PIEPreviewDeviceSpecification.h
--- a/Engine/Source/Editor/PIEPreviewDeviceSpecification/Public/PIEPreviewDeviceSpecification.h
+++ b/Engine/Source/Editor/PIEPreviewDeviceSpecification/Public/PIEPreviewDeviceSpecification.h
@@ -109,7 +109,8 @@ public:
 
 	UPROPERTY()
 	FPIERHIOverrideState GLES31RHIState;
-
+	UPROPERTY()
+	bool SM5Available = false;
 // 	UPROPERTY()
 // 	FPIERHIOverrideState VulkanRHIState;
 };
diff --git a/Engine/Source/Runtime/Core/Private/Android/AndroidPlatformMisc.cpp b/Engine/Source/Runtime/Core/Private/Android/AndroidPlatformMisc.cpp
--- a/Engine/Source/Runtime/Core/Private/Android/AndroidPlatformMisc.cpp
+++ b/Engine/Source/Runtime/Core/Private/Android/AndroidPlatformMisc.cpp
@@ -2307,7 +2307,7 @@ static void EstablishVulkanDeviceSupport()
 	}
 }
 
-bool IsDesktopVulkanAvailable()
+bool FAndroidMisc::IsDesktopVulkanAvailable()
 {
 	static int CachedDesktopVulkanAvailable = -1;
 
diff --git a/Engine/Source/Runtime/Core/Public/Android/AndroidPlatformMisc.h b/Engine/Source/Runtime/Core/Public/Android/AndroidPlatformMisc.h
--- a/Engine/Source/Runtime/Core/Public/Android/AndroidPlatformMisc.h
+++ b/Engine/Source/Runtime/Core/Public/Android/AndroidPlatformMisc.h
@@ -200,6 +200,7 @@ public:
 	 *			and not we are not forcing GLES with a command line switch
 	 */
 	static CORE_API bool IsVulkanAvailable();
+	static CORE_API bool IsDesktopVulkanAvailable();
 
 	/* ShouldUseVulkan
 	 * @return true if Vulkan is available, and not disabled by device profile cvar
