Fix Android SM5 matching rules for Mali Immortalis

diff --git a/Engine/Config/BaseDeviceProfiles.ini b/Engine/Config/BaseDeviceProfiles.ini
--- a/Engine/Config/BaseDeviceProfiles.ini
+++ b/Engine/Config/BaseDeviceProfiles.ini
@@ -736,17 +736,16 @@ BaseProfileName=VisionPro
 +MatchProfile=(Profile="Android_Mali_G78",Match=((SourceType=SRC_GpuFamily,CompareType=CMP_Regex,MatchString="^Mali\\-G78($| )")))
 +MatchProfile=(Profile="Android_Mali_G710_Vulkan",Match=((SourceType=SRC_GpuFamily,CompareType=CMP_Regex,MatchString="^Mali\\-G710"),(SourceType=SRC_VulkanAvailable,CompareType=CMP_Equal,MatchString="true")))
 +MatchProfile=(Profile="Android_Mali_G710",Match=((SourceType=SRC_GpuFamily,CompareType=CMP_Regex,MatchString="^Mali\\-G710")))
++MatchProfile=(Profile="Android_Vulkan_SM5_Mali",Match=((SourceType=SRC_GpuFamily,CompareType=CMP_Regex,MatchString="^Mali\\-G7[1-9][0-9]"),(SourceType=SRC_SM5Available,CompareType=CMP_Equal,MatchString="true")))
 +MatchProfile=(Profile="Android_Mali_G7xx_Vulkan",Match=((SourceType=SRC_GpuFamily,CompareType=CMP_Regex,MatchString="^Mali\\-G7[0-9][0-9]"),(SourceType=SRC_VulkanAvailable,CompareType=CMP_Equal,MatchString="true")))
 +MatchProfile=(Profile="Android_Mali_G7xx",Match=((SourceType=SRC_GpuFamily,CompareType=CMP_Regex,MatchString="^Mali\\-G7[0-9][0-9]")))
-; Example profiles for Vulkan SM5
-;+MatchProfile=(Profile="Android_Vulkan_SM5_Mali",Match=((SourceType=SRC_GpuFamily,CompareType=CMP_Regex,MatchString="^Mali\\-G7[67]"),(SourceType=SRC_SM5Available,CompareType=CMP_Equal,MatchString="true")))
-+MatchProfile=(Profile="Android_Vulkan_SM5_Mali",Match=((SourceType=SRC_GpuFamily,CompareType=CMP_Regex,MatchString="^Mali\\-G7[1-9][0-9]"),(SourceType=SRC_SM5Available,CompareType=CMP_Equal,MatchString="true")))
+
 ; Enable Vulkan on any unknown Mali running Android 10 or later
 +MatchProfile=(Profile="Android_Unknown_Vulkan",Match=((SourceType=SRC_GpuFamily,CompareType=CMP_Regex,MatchString="^Mali"),(SourceType=SRC_AndroidVersion, CompareType=CMP_Regex,MatchString="([0-9]+).*"),(SourceType=SRC_PreviousRegexMatch,CompareType=CMP_GreaterEqual,MatchString="10"),(SourceType=SRC_VulkanAvailable,CompareType=CMP_Equal,MatchString="true")))
 
 ; Samsung Exynos / Xclipse GPU-based devices use Vulkan by default.
-;+MatchProfile=(Profile="Android_Vulkan_SM5",Match=((SourceType=SRC_GpuFamily,CompareType=CMP_Regex,MatchString="Samsung Xclipse 920"),(SourceType=SRC_VulkanAvailable,CompareType=CMP_Equal,MatchString="true")))
-;+MatchProfile=(Profile="Android_Vulkan_SM5_Mali",Match=((SourceType=SRC_GpuFamily,CompareType=CMP_Regex,MatchString="^Mali\\-G7[67]"),(SourceType=SRC_VulkanAvailable,CompareType=CMP_Equal,MatchString="true")))
+;+MatchProfile=(Profile="Android_Vulkan_SM5",Match=((SourceType=SRC_GpuFamily,CompareType=CMP_Regex,MatchString="Samsung Xclipse 920"),(SourceType=SRC_SM5Available,CompareType=CMP_Equal,MatchString="true")))
+;+MatchProfile=(Profile="Android_Vulkan_SM5_Mali",Match=((SourceType=SRC_GpuFamily,CompareType=CMP_Regex,MatchString="^Mali\\-G7[67]"),(SourceType=SRC_SM5Available,CompareType=CMP_Equal,MatchString="true")))
 +MatchProfile=(Profile="Android_Xclipse_5xx_Vulkan",Match=((SourceType=SRC_GpuFamily,CompareType=CMP_Regex,MatchString="Samsung Xclipse 5[0-9][0-9]"),(SourceType=SRC_VulkanAvailable,CompareType=CMP_Equal,MatchString="true")))
 +MatchProfile=(Profile="Android_Xclipse_5xx",Match=((SourceType=SRC_GpuFamily,CompareType=CMP_Regex,MatchString="Samsung Xclipse 5[0-9][0-9]")))
 +MatchProfile=(Profile="Android_Xclipse_9xx_Vulkan",Match=((SourceType=SRC_GpuFamily,CompareType=CMP_Regex,MatchString="Samsung Xclipse 9[0-9][0-9]"),(SourceType=SRC_VulkanAvailable,CompareType=CMP_Equal,MatchString="true")))
