Add a few commands to the console autocomplete list:  `obj gc`, `obj dump`, `debug crash`, `debug rendercrash` and extended `obj refs`

diff --git a/Engine/Config/BaseInput.ini b/Engine/Config/BaseInput.ini
--- a/Engine/Config/BaseInput.ini
+++ b/Engine/Config/BaseInput.ini
@@ -67,6 +67,8 @@ MaxScrollbackSize=1024
 +ManualAutoCompleteList=(Command="RecompileShaders global",Desc="Recompile global shaders that have any changes on their source files")
 +ManualAutoCompleteList=(Command="RecompileShaders material ",Desc="Recompile shaders for a specific material if it's source files have changed")
 +ManualAutoCompleteList=(Command="RecompileShaders all",Desc="Recompile all shaders that have any changes on their source files")
++ManualAutoCompleteList=(Command="Debug Crash",Desc="Simulates a game thread crash for debugging")
++ManualAutoCompleteList=(Command="Debug RenderCrash",Desc="Simulates a render thread crash for debugging")
 +ManualAutoCompleteList=(Command="DumpMaterialStats",Desc="Dump material information")
 +ManualAutoCompleteList=(Command="DumpShaderStats",Desc="Dump shader information")
 +ManualAutoCompleteList=(Command="DumpShaderPipelineStats",Desc="Dump shader pipeline information")
@@ -89,7 +91,9 @@ MaxScrollbackSize=1024
 +ManualAutoCompleteList=(Command="Obj List ",Desc="<Class=ClassName> <Type=MetaClass> <Outer=OuterObject> <Package=InsidePackage> <Inside=InsideObject>")
 +ManualAutoCompleteList=(Command="Obj ListContentRefs",Desc="<Class=ClassName> <ListClass=ClassName>")
 +ManualAutoCompleteList=(Command="Obj Classes",Desc="Shows all classes")
-+ManualAutoCompleteList=(Command="Obj Refs",Desc="Name=<ObjectName> Class=<OptionalObjectClass> Lists referencers of the specified object")
++ManualAutoCompleteList=(Command="Obj Refs",Desc="Name=<ObjectName> [Shortest] [Longest] [All] [External] [Direct] [Full] [Minimal] [GCOnly] [History=<MaxHistoryLevel>] Lists referencers of the specified object")
++ManualAutoCompleteList=(Command="Obj GC",Desc="Runs the UObject Garbage Collector and resets the GC timer.")
++ManualAutoCompleteList=(Command="Obj Dump ",Desc="<ObjectName> [Recurse or Hide/Show=\"category1,...\"] or <Class=ClassName> <Name=ObjectName> [Recurse or Hide/Show=\"category1,...\"] Prints the value of all variables for the specified object")
 +ManualAutoCompleteList=(Command="EditActor",Desc="<Class=ClassName> or <Name=ObjectName> or TRACE")
 +ManualAutoCompleteList=(Command="EditDefault",Desc="<Class=ClassName>")
 +ManualAutoCompleteList=(Command="EditObject",Desc="<Class=ClassName> or <Name=ObjectName> or <ObjectName>")
