Fix HttpActivityTimeout to be lower than HttpConnectionTimeout

diff --git a/Engine/Config/BaseEngine.ini b/Engine/Config/BaseEngine.ini
--- a/Engine/Config/BaseEngine.ini
+++ b/Engine/Config/BaseEngine.ini
@@ -29,7 +29,7 @@ Port=7777
 
 [HTTP]
 HttpConnectionTimeout=60
-HttpActivityTimeout=180
+HttpActivityTimeout=30
 ; Flush time limits for FHttpManager::Flush:
 ; Default = when Flush is called with no overriding context
 ; Background = when Flush is called because the app is backgrounding.  This may not be called on backgrounding on all platforms.
