Iris
* Rename Prioritizer configs to make them easier to search for

diff --git a/Engine/Config/BaseEngine.ini b/Engine/Config/BaseEngine.ini
--- a/Engine/Config/BaseEngine.ini
+++ b/Engine/Config/BaseEngine.ini
@@ -1307,8 +1307,8 @@ MaxPartCount=4096
 +DataStreamDefinitions=(DataStreamName=Replication, ClassName=/Script/IrisCore.ReplicationDataStream, DefaultSendStatus=EDataStreamSendStatus::Send, bAutoCreate=true)
 
 [/Script/IrisCore.NetObjectPrioritizerDefinitions]
-+NetObjectPrioritizerDefinitions=(PrioritizerName=Default, ClassName=/Script/IrisCore.SphereNetObjectPrioritizer, ConfigClassName=/Script/IrisCore.SphereNetObjectPrioritizerConfig)
-+NetObjectPrioritizerDefinitions=(PrioritizerName=PlayerState, ClassName=/Script/IrisCore.NetObjectCountLimiter, ConfigClassName=/Script/Engine.PlayerStateCountLimiterConfig)
++NetObjectPrioritizerDefinitions=(PrioritizerName=DefaultPrioritizer, ClassName=/Script/IrisCore.SphereNetObjectPrioritizer, ConfigClassName=/Script/IrisCore.SphereNetObjectPrioritizerConfig)
++NetObjectPrioritizerDefinitions=(PrioritizerName=PlayerStatePrioritizer, ClassName=/Script/IrisCore.NetObjectCountLimiter, ConfigClassName=/Script/Engine.PlayerStateCountLimiterConfig)
 
 [/Script/IrisCore.NetObjectFilterDefinitions]
 +NetObjectFilterDefinitions=(FilterName=Spatial, ClassName=/Script/IrisCore.NetObjectGridWorldLocFilter, ConfigClassName=/Script/IrisCore.NetObjectGridFilterConfig)
diff --git a/Engine/Source/Runtime/Experimental/Iris/Core/Private/Iris/ReplicationSystem/Prioritization/ReplicationPrioritization.cpp b/Engine/Source/Runtime/Experimental/Iris/Core/Private/Iris/ReplicationSystem/Prioritization/ReplicationPrioritization.cpp
--- a/Engine/Source/Runtime/Experimental/Iris/Core/Private/Iris/ReplicationSystem/Prioritization/ReplicationPrioritization.cpp
+++ b/Engine/Source/Runtime/Experimental/Iris/Core/Private/Iris/ReplicationSystem/Prioritization/ReplicationPrioritization.cpp
@@ -20,6 +20,8 @@ namespace UE::Net::Private
 static_assert(InvalidNetObjectPrioritizerHandle == ~FNetObjectPrioritizerHandle(0), "ObjectIndexToPrioritizer code needs attention. Contact the UE Networking team.");
 static constexpr uint8 FReplicationPrioritization_InvalidNetObjectPrioritizerIndex = 0xFF;
 
+static const FName NAME_DefaultPrioritizer(TEXT("DefaultPrioritizer"));
+
 /**
  * Most logic in here revolves around batches. As such we need access to the Chunks.
  */
@@ -424,7 +426,7 @@ FNetObjectPrioritizerHandle FReplicationPrioritization::GetPrioritizerHandle(con
 		++Handle;
 	}
 
-	if (PrioritizerName == NAME_Default)
+	if (PrioritizerName == UE::Net::Private::NAME_DefaultPrioritizer)
 	{
 		return DefaultSpatialNetObjectPrioritizerHandle;
 	}
diff --git a/Engine/Source/Runtime/Experimental/Iris/Core/Public/Iris/ReplicationSystem/Prioritization/NetObjectCountLimiter.h b/Engine/Source/Runtime/Experimental/Iris/Core/Public/Iris/ReplicationSystem/Prioritization/NetObjectCountLimiter.h
--- a/Engine/Source/Runtime/Experimental/Iris/Core/Public/Iris/ReplicationSystem/Prioritization/NetObjectCountLimiter.h
+++ b/Engine/Source/Runtime/Experimental/Iris/Core/Public/Iris/ReplicationSystem/Prioritization/NetObjectCountLimiter.h
@@ -23,6 +23,7 @@ enum class ENetObjectCountLimiterMode : uint32
 	Fill,
 };
 
+//TODO $IRIS: Document class usage
 UCLASS(transient, config=Engine, MinimalAPI)
 class UNetObjectCountLimiterConfig : public UNetObjectPrioritizerConfig
 {
