Updated ini file to include an example Horde Analytics setup

diff --git a/Engine/Config/BaseEngine.ini b/Engine/Config/BaseEngine.ini
--- a/Engine/Config/BaseEngine.ini
+++ b/Engine/Config/BaseEngine.ini
@@ -3553,7 +3553,9 @@ IntervalSeconds=1
 +Listeners=(MetricName=AddedConnections, Class=/Script/Engine.NetworkMetricsPerfCounters)
 
 ; Studio Telemetry Settings
-[StudioTelemetry.Log] 
+
+; Example Log File Analytics
+[StudioTelemetry.LogAnalytics] 
 Name=LogAnalytics
 ProviderModule=AnalyticsLog
 UsageType=Editor
@@ -3561,6 +3563,15 @@ UsageType=Editor
 ; FileName="somefile.json"
 ; FolderPath="somefolder"
 
+; Example Horde Analytics Settings
+;[StudioTelemetry.HordeAnalytics]
+; Name=HordeAnalytics
+; ProviderModule=AnalyticsET
+; UsageType=EditorAndClient
+; APIKeyET=HordeAnalytics.Dev
+; APIServerET="https://myownhordeserver.com/" ; Point Server URL to your own Horde server
+; APIEndpointET="api/v1/telemetry"
+
 [/Script/Engine.WorldPartitionSettings]
 EditorHashDefaultClass=/Script/CoreUObject.Class'/Script/Engine.WorldPartitionEditorSpatialHash'
 RuntimeHashDefaultClass=/Script/CoreUObject.Class'/Script/Engine.WorldPartitionRuntimeHashSet'
