Disabling Zen as Cooked Output Store by default for licensee configurations.  Still remains enabled by default for non-licensee configurations.

diff --git a/Engine/Config/BaseGame.ini b/Engine/Config/BaseGame.ini
--- a/Engine/Config/BaseGame.ini
+++ b/Engine/Config/BaseGame.ini
@@ -89,7 +89,7 @@ BuildConfiguration=PPBC_Development
 FullRebuild=False
 UsePakFile=True
 bUseIoStore=True
-bUseZenStore=True
+bUseZenStore=False
 bCompressed=True
 PackageCompressionFormat=Oodle
 ; default compression block size is 64 KB
