Moved HTTP Threaded Request limit out of general config and into PS5 specific config. Updated rate to the maximum without errors after testing.

diff --git a/Engine/Config/BaseEngine.ini b/Engine/Config/BaseEngine.ini
--- a/Engine/Config/BaseEngine.ini
+++ b/Engine/Config/BaseEngine.ini
@@ -46,9 +46,6 @@ FlushHardTimeLimitShutdown=4.0
 FlushSoftTimeLimitFullFlush=-1.0
 FlushHardTimeLimitFullFlush=-1.0
 
-[HTTP.HttpThread]
-RunningThreadedRequestLimit=11
-
 [HTTP.Curl]
 bAllowSeekFunction=false
 
