Actor Coloration Framework: added "Affects NavMesh" property visualization in the editor.

diff --git a/Engine/Config/BaseEditorSettings.ini b/Engine/Config/BaseEditorSettings.ini
--- a/Engine/Config/BaseEditorSettings.ini
+++ b/Engine/Config/BaseEditorSettings.ini
@@ -44,6 +44,9 @@ IncludeAssetPaths=True
 ; Whether to include collection names when searching by text
 IncludeCollectionNames=True
 
+[/Script/UnrealEd.PropertyColorSettings]
++CustomProperties=(Name="AffectsNavmesh",Text="Affects NavMesh",PropertyChain="RootComponent.bCanEverAffectNavigation",PropertyValue="True",PropertyColor=(R=0,G=0,B=255))
+
 [/Script/VREditor.VRModeSettings]
 bEnableAutoVREditMode=False
 bAutokeySequences=True
