AssetRegistry: Disable Discovery cache and use of PlatformFileJournal. It needs to be fixed to work with symbolic links and with some Windows11 edge cases.
#rnx

diff --git a/Engine/Config/BaseEngine.ini b/Engine/Config/BaseEngine.ini
--- a/Engine/Config/BaseEngine.ini
+++ b/Engine/Config/BaseEngine.ini
@@ -2733,7 +2733,7 @@ bFilterSearchableNames=true
 ; by setting AssetRegistryDiscoveryCache=AlwaysWrite, and it can be force-read with no test for invalidation due to modified
 ; directories by commandline -AssetRegistryCacheSkipInvalidate.
 ; Expected values for AssetRegistryDiscoveryCache are Never,Default,AlwaysWrite.
-AssetRegistryDiscoveryCache=Default
+AssetRegistryDiscoveryCache=Never
 
 [AutomationTesting]
 ImportTestPath=Content/EditorAutomation/
