Magic Leap 2 fixes

diff --git a/Engine/Config/BaseDeviceProfiles.ini b/Engine/Config/BaseDeviceProfiles.ini
--- a/Engine/Config/BaseDeviceProfiles.ini
+++ b/Engine/Config/BaseDeviceProfiles.ini
@@ -700,7 +700,7 @@ BaseProfileName=VisionPro
 +MatchProfile=(Profile="Meta_Quest_Pro",Match=((SourceType=SRC_DeviceMake,CompareType=CMP_Equal,MatchString="Oculus"), (SourceType=SRC_HMDSystemName,CompareType=CMP_Equal,MatchString="Meta Quest Pro")))
 +MatchProfile=(Profile="Oculus_Quest2",Match=((SourceType=SRC_DeviceMake,CompareType=CMP_Equal,MatchString="Oculus"), (SourceType=SRC_HMDSystemName,CompareType=CMP_Equal,MatchString="Oculus Quest2")))
 +MatchProfile=(Profile="Oculus_Quest",Match=((SourceType=SRC_DeviceMake,CompareType=CMP_Equal,MatchString="Oculus")))
-+MatchProfile=(Profile="MagicLeap_Vulkan",Match=((SourceType=SRC_DeviceMake,CompareType=CMP_Equal,MatchString="MagicLeap"), (SourceType=SRC_VulkanAvailable,CompareType=CMP_Equal,MatchString="true")))
++MatchProfile=(Profile="MagicLeap_Vulkan",Match=((SourceType=SRC_DeviceMake,CompareType=CMP_Equal,MatchString="Magic Leap"), (SourceType=SRC_VulkanAvailable,CompareType=CMP_Equal,MatchString="true")))
 
 ; Match Android emulator
 +MatchProfile=(Profile="Android_Emulator",Match=((SourceType=SRC_DeviceMake,CompareType=CMP_Equal,MatchString="Google"),(SourceType=SRC_GPUFamily,CompareType=CMP_Regex,MatchString="Android Emulator")))
@@ -1195,3 +1195,5 @@ BaseProfileName=Windows
 DeviceType=Android
 BaseProfileName=Android_Mid
 +CVars=xr.DisableOpenXROnAndroidWithoutOculus=0
++CVars=r.Android.DisableVulkanSupport=0
++CVars=r.Vulkan.SupportsBCTextureFormats=1
diff --git a/Engine/Source/Runtime/VulkanRHI/Private/Android/VulkanAndroidPlatform.cpp b/Engine/Source/Runtime/VulkanRHI/Private/Android/VulkanAndroidPlatform.cpp
--- a/Engine/Source/Runtime/VulkanRHI/Private/Android/VulkanAndroidPlatform.cpp
+++ b/Engine/Source/Runtime/VulkanRHI/Private/Android/VulkanAndroidPlatform.cpp
@@ -68,6 +68,15 @@ static TAutoConsoleVariable<int32> CVarVulkanSupportsTimestampQueries(
 	ECVF_SetByDeviceProfile
 );
 
+static TAutoConsoleVariable<int32> CVarVulkanSupportsBCTextureFormats(
+	TEXT("r.Vulkan.SupportsBCTextureFormats"),
+	0,
+	TEXT("Whether or not BC Texture formats are supported\n")
+	TEXT("  0 = unsupported\n")
+	TEXT("  1 = supported."),
+	ECVF_SetByDeviceProfile
+);
+
 // Vulkan function pointers
 #define DEFINE_VK_ENTRYPOINTS(Type,Func) Type VulkanDynamicAPI::Func = NULL;
 ENUM_VK_ENTRYPOINTS_ALL(DEFINE_VK_ENTRYPOINTS)
@@ -404,6 +413,11 @@ void* FVulkanAndroidPlatform::GetHardwareWindowHandle()
 	return WindowHandle;
 }
 
+bool FVulkanAndroidPlatform::SupportsBCTextureFormats()
+{
+	return (CVarVulkanSupportsBCTextureFormats.GetValueOnAnyThread() == 1);
+}
+
 void FVulkanAndroidPlatform::CreateSurface(void* WindowHandle, VkInstance Instance, VkSurfaceKHR* OutSurface)
 {
 	// don't use cached window handle coming from VulkanViewport, as it could be gone by now
diff --git a/Engine/Source/Runtime/VulkanRHI/Private/Android/VulkanAndroidPlatform.h b/Engine/Source/Runtime/VulkanRHI/Private/Android/VulkanAndroidPlatform.h
--- a/Engine/Source/Runtime/VulkanRHI/Private/Android/VulkanAndroidPlatform.h
+++ b/Engine/Source/Runtime/VulkanRHI/Private/Android/VulkanAndroidPlatform.h
@@ -59,7 +59,7 @@ public:
 
 	static void* GetHardwareWindowHandle();
 
-	static bool SupportsBCTextureFormats() { return false; }
+	static bool SupportsBCTextureFormats();
 	static bool SupportsASTCTextureFormats() { return true; }
 	static bool SupportsETC2TextureFormats() { return true; }
 	// GLES does not support R16Unorm, so all Android has to fallback to R16F instead
