Revert changelist which changed the concurrent http requests limit on release branch, which is not ready yet.

diff --git a/Engine/Config/BaseEngine.ini b/Engine/Config/BaseEngine.ini
--- a/Engine/Config/BaseEngine.ini
+++ b/Engine/Config/BaseEngine.ini
@@ -46,6 +46,9 @@ FlushHardTimeLimitShutdown=4.0
 FlushSoftTimeLimitFullFlush=-1.0
 FlushHardTimeLimitFullFlush=-1.0
 
+[HTTP.HttpThread]
+RunningThreadedRequestLimit=11
+
 [HTTP.Curl]
 bAllowSeekFunction=false
 
