Unshelved from pending changelist '31851286':

- Add the turnkey platforms menu settings to be saved properly

diff --git a/Engine/Config/BaseGame.ini b/Engine/Config/BaseGame.ini
--- a/Engine/Config/BaseGame.ini
+++ b/Engine/Config/BaseGame.ini
@@ -4,6 +4,7 @@
 
 [SectionsToSave]
 +Section=DebugWindows
++Section=/Script/DeveloperToolSettings.PlatformsMenuSettings
 
 [Internationalization]
 ; To add paths to use for localized game text, use the localization tools or add more lines like: +LocalizationPaths=%GAMEDIR%Content/Localization/SpecificConfiguration
