Split engine.ini:[Ias]:ForceTocFromMountedPaks into two different values.

#rnx

- When it was added ForceTocFromMountedPaks controlled if we should create .uondemandtoc files when building IAS data and if we should use them at runtime, this was inflexiable for testing purposes so has been split into:
-- engine.ini:[Ias]:CreatePerContainerTocs which controls if we should create the files.
-- engine.ini:[Ias]:UsePerContainerTocs which controls if we should run from the files.
- Both values remain defaulted to off so that the system is only used as opt in.

diff --git a/Engine/Config/BaseEngine.ini b/Engine/Config/BaseEngine.ini
--- a/Engine/Config/BaseEngine.ini
+++ b/Engine/Config/BaseEngine.ini
@@ -3484,7 +3484,9 @@ bUpdateBuildGraphPropertiesFile=false
 
 [Ias]
 EnableInEditor=False
-ForceTocFromMountedPaks=False
+; Note that the following two entries are temp while we test a new feature
+CreatePerContainerTocs=False
+UsePerContainerTocs=False
 
 [/Script/InterchangeEngine.InterchangeProjectSettings]
 ContentImportSettings=(DefaultPipelineStackOverride=((Materials, "Materials"),(Textures, "Textures")),ShowImportDialogOverride=((Textures, (bShowImportDialog=False))),PipelineStacks=(("Assets", (Pipelines=("/Interchange/Pipelines/DefaultAssetsPipeline.DefaultAssetsPipeline"),PerTranslatorPipelines=((Translator="/Script/InterchangeImport.InterchangeMaterialXTranslator",Pipelines=("/Interchange/Pipelines/DefaultAssetsPipeline.DefaultAssetsPipeline","/Interchange/Pipelines/DefaultMaterialXPipeline.DefaultMaterialXPipeline")),(Translator="/Script/InterchangeImport.InterchangeGLTFTranslator",Pipelines=("/Interchange/Pipelines/DefaultGLTFAssetsPipeline.DefaultGLTFAssetsPipeline","/Interchange/Pipelines/DefaultGLTFPipeline.DefaultGLTFPipeline"))))),("Materials", (Pipelines=("/Interchange/Pipelines/DefaultMaterialPipeline.DefaultMaterialPipeline"),PerTranslatorPipelines=((Translator="/Script/InterchangeImport.InterchangeMaterialXTranslator",Pipelines=("/Interchange/Pipelines/DefaultMaterialPipeline.DefaultMaterialPipeline","/Interchange/Pipelines/DefaultMaterialXPipeline.DefaultMaterialXPipeline"))))),("Textures", (Pipelines=("/Interchange/Pipelines/DefaultTexturePipeline.DefaultTexturePipeline")))),DefaultPipelineStack="Assets",ImportDialogClass="/Script/InterchangeEditorPipelines.InterchangePipelineConfigurationGeneric",bShowImportDialog=True)
diff --git a/Engine/Source/Runtime/Experimental/IoStoreOnDemand/Private/IoStoreOnDemand.cpp b/Engine/Source/Runtime/Experimental/IoStoreOnDemand/Private/IoStoreOnDemand.cpp
--- a/Engine/Source/Runtime/Experimental/IoStoreOnDemand/Private/IoStoreOnDemand.cpp
+++ b/Engine/Source/Runtime/Experimental/IoStoreOnDemand/Private/IoStoreOnDemand.cpp
@@ -1302,7 +1302,7 @@ TIoStatusOr<FIoStoreUploadResult> UploadContainerFiles(
 	ChunksRelativePath.ToLowerInline();
 
 	bool bWritePerContainerToc = false;
-	GConfig->GetBool(TEXT("Ias"), TEXT("ForceTocFromMountedPaks"), bWritePerContainerToc, GEngineIni);
+	GConfig->GetBool(TEXT("Ias"), TEXT("CreatePerContainerTocs"), bWritePerContainerToc, GEngineIni);
 
 	uint64 TotalUploadedChunks = 0;
 	uint64 TotalUploadedBytes = 0;
diff --git a/Engine/Source/Runtime/Experimental/IoStoreOnDemand/Private/OnDemandIoDispatcherBackend.cpp b/Engine/Source/Runtime/Experimental/IoStoreOnDemand/Private/OnDemandIoDispatcherBackend.cpp
--- a/Engine/Source/Runtime/Experimental/IoStoreOnDemand/Private/OnDemandIoDispatcherBackend.cpp
+++ b/Engine/Source/Runtime/Experimental/IoStoreOnDemand/Private/OnDemandIoDispatcherBackend.cpp
@@ -1483,10 +1483,10 @@ FOnDemandIoBackend::FOnDemandIoBackend(TUniquePtr<IIasCache>&& InCache)
 #endif // UE_IAS_DEBUG_CONSOLE_CMDS
 
 	{
-		bool bForceTocFromMountedPaks = false;
-		GConfig->GetBool(TEXT("Ias"), TEXT("ForceTocFromMountedPaks"), bForceTocFromMountedPaks, GEngineIni);
+		bool bUsePerContainerTocs = false;
+		GConfig->GetBool(TEXT("Ias"), TEXT("UsePerContainerTocs"), bUsePerContainerTocs, GEngineIni);
 
-		if (bForceTocFromMountedPaks)
+		if (bUsePerContainerTocs)
 		{
 			ForceTocMode(ETocMode::LoadTocFromMountedPaks);
 		}
