Add NoZenLocalFallback and InstalledNoZenLocalFallback DDC graphs for use in scenarios where users are unable to run a local zenserver for any reason.  Also fix a mis-labeling of a local cache store as being remote if it is deactivated.

diff --git a/Engine/Config/BaseEngine.ini b/Engine/Config/BaseEngine.ini
--- a/Engine/Config/BaseEngine.ini
+++ b/Engine/Config/BaseEngine.ini
@@ -2555,6 +2555,26 @@ EnterprisePak=(Type=ReadPak, Filename=../../../Enterprise/DerivedDataCache/Compr
 ZenLocal=(Type=Zen, ServerID=Local)
 ZenShared=(Type=Zen, ServerID=Shared)
 
+[NoZenLocalFallback]
+Root=(Type=Hierarchical, Inner=Pak, Inner=EnginePak, Inner=Local, Inner=ZenShared, Inner=Shared, Inner=Cloud)
+Local=(Type=FileSystem, UnusedFileAge=34, PromptIfMissing=true, Path=%ENGINEDIR%DerivedDataCache, EnvPathOverride=UE-LocalDataCachePath, EditorOverrideSetting=LocalDerivedDataCache, CommandLineOverride=LocalDataCachePath)
+Shared=(Type=FileSystem, UnusedFileAge=10, FoldersToClean=10, ConsiderSlowAt=70, Path=?EpicDDC, EnvPathOverride=UE-SharedDataCachePath, EditorOverrideSetting=SharedDerivedDataCache, CommandLineOverride=SharedDataCachePath)
+Cloud=(Type=Cloud, ServerID=Cloud, EnvHostOverride=UE-CloudDataCacheHost, CommandLineHostOverride=CloudDataCacheHost, CommandLineOAuthSecretOverride="CloudDataCacheOAuthSecret", OAuthAccessTokenEnvOverride="UE-CloudDataCacheAccessToken", EnvHttpVersionOverride="UE-CloudDataCacheHttpVersion", CommandLineHttpVersionOverride="CloudDataCacheHttpVersion")
+Pak=(Type=ReadPak, Filename="%GAMEDIR%DerivedDataCache/DDC.ddp")
+EnginePak=(Type=ReadPak, Filename="%ENGINEDIR%DerivedDataCache/DDC.ddp")
+ZenShared=(Type=Zen, ServerID=Shared)
+
+[InstalledNoZenLocalFallback]
+Root=(Type=Hierarchical, Inner=Pak, Inner=CompressedPak, Inner=EnginePak, Inner=EnterprisePak, Inner=Local, Inner=ZenShared, Inner=Shared, Inner=Cloud)
+Local=(Type=FileSystem, UnusedFileAge=34, Path="%ENGINEVERSIONAGNOSTICUSERDIR%DerivedDataCache", EnvPathOverride=UE-LocalDataCachePath, EditorOverrideSetting=LocalDerivedDataCache)
+Shared=(Type=FileSystem, UnusedFileAge=10, FoldersToClean=10, ConsiderSlowAt=70, Path=?EpicDDC, EnvPathOverride=UE-SharedDataCachePath, EditorOverrideSetting=SharedDerivedDataCache, CommandLineOverride=SharedDataCachePath)
+Cloud=(Type=Cloud, ServerID=Cloud, EnvHostOverride=UE-CloudDataCacheHost, CommandLineHostOverride=CloudDataCacheHost, CommandLineOAuthSecretOverride="CloudDataCacheOAuthSecret", OAuthAccessTokenEnvOverride="UE-CloudDataCacheAccessToken", EnvHttpVersionOverride="UE-CloudDataCacheHttpVersion", CommandLineHttpVersionOverride="CloudDataCacheHttpVersion")
+Pak=(Type=ReadPak, Filename="%GAMEDIR%DerivedDataCache/DDC.ddp")
+CompressedPak=(Type=ReadPak, Filename="%GAMEDIR%DerivedDataCache/Compressed.ddp", Compressed=true)
+EnginePak=(Type=ReadPak, Filename=../../../Engine/DerivedDataCache/Compressed.ddp, Compressed=true)
+EnterprisePak=(Type=ReadPak, Filename=../../../Enterprise/DerivedDataCache/Compressed.ddp, Compressed=true)
+ZenShared=(Type=Zen, ServerID=Shared)
+
 [ZenOnlyDDC]
 Root=(Type=KeyLength, Length=120, Inner=AsyncPut)
 AsyncPut=(Type=AsyncPut, Inner=Zen)
diff --git a/Engine/Source/Developer/DerivedDataCache/Private/ZenCacheStore.cpp b/Engine/Source/Developer/DerivedDataCache/Private/ZenCacheStore.cpp
--- a/Engine/Source/Developer/DerivedDataCache/Private/ZenCacheStore.cpp
+++ b/Engine/Source/Developer/DerivedDataCache/Private/ZenCacheStore.cpp
@@ -1688,7 +1688,7 @@ void FZenCacheStore::Initialize(const FZenCacheStoreParams& Params)
 	ClientParams.LowSpeedTime = 25;
 	RequestQueue = FHttpRequestQueue(*ConnectionPool, ClientParams);
 
-	bIsLocalConnection = ZenService.GetInstance().IsServiceRunningLocally();
+	bIsLocalConnection = ZenService.GetInstance().IsServiceRunningLocally() || ZenService.GetInstance().GetServiceSettings().IsAutoLaunch();
 	bIsUsable = true;
 
 
@@ -1712,13 +1712,26 @@ void FZenCacheStore::Initialize(const FZenCacheStoreParams& Params)
 
 		if (!bReady)
 		{
-			Flags = ECacheStoreFlags::None;
-			UE_LOG(LogDerivedDataCache, Display,
-				TEXT("%s: Readiness check failed. "
-					"It will be deactivated until responsiveness improves. "
-					"If this is consistent, consider disabling this cache store through "
-					"environment variables or other configuration."),
-				*GetName());
+			Flags = OperationalFlags & ~(ECacheStoreFlags::Store | ECacheStoreFlags::Query);
+			if (bIsLocalConnection)
+			{
+				UE_LOG(LogDerivedDataCache, Display,
+					TEXT("%s: Readiness check failed. "
+						"It will be deactivated until responsiveness improves. "
+						"If this is consistent, consider disabling this cache store through "
+						"the use of the '-ddc=NoZenLocalFallback' or '-ddc=InstalledNoZenLocalFallback' "
+						"commandline arguments."),
+					*GetName());
+			}
+			else
+			{
+				UE_LOG(LogDerivedDataCache, Display,
+					TEXT("%s: Readiness check failed. "
+						"It will be deactivated until responsiveness improves. "
+						"If this is consistent, consider disabling this cache store through "
+						"environment variables or other configuration."),
+					*GetName());
+			}
 			bDeactivatedForPerformance.store(true, std::memory_order_relaxed);
 		}
 
@@ -1910,7 +1923,7 @@ void FZenCacheStore::ConditionalEvaluatePerformance()
 	{
 		if (!bDeactivatedForPerformance.load(std::memory_order_relaxed))
 		{
-			StoreOwner->SetFlags(this, ECacheStoreFlags::None);
+			StoreOwner->SetFlags(this, OperationalFlags & ~(ECacheStoreFlags::Store | ECacheStoreFlags::Query));
 			UE_LOG(LogDerivedDataCache, Display,
 				TEXT("%s: Performance does not meet minimum criteria. "
 					"It will be deactivated until performance measurements improve. "
