Remove enabling of r.VirtualTextures from Android SM5 default config.

Currently disabling VT in Project Settings breaks Android SM5 renderer because shaders are missing.
By removing  r.VirtualTextures=1 from the config it's now possibler to toggle it via Project Settings.
It is still possible to use VT for SM5 but have it disabled for mobile since mobile already requires an additional mobile CVar to be enabled.

diff --git a/Engine/Config/BaseDeviceProfiles.ini b/Engine/Config/BaseDeviceProfiles.ini
--- a/Engine/Config/BaseDeviceProfiles.ini
+++ b/Engine/Config/BaseDeviceProfiles.ini
@@ -882,7 +882,6 @@ BaseProfileName=Android
 +CVars=r.SSS.Quality=0
 +CVars=r.SSS.Scale=0
 +CVars=r.SSR.Quality=0
-+CVars=r.VirtualTextures=1
 +CVars=r.Android.DisableVulkanSM5Support=0
 +CVars=r.Android.DisableVulkanSupport=0
 +CVars=r.DistanceFields=1
