Add Android SM5 device profile for Xclipse devices
Rearrange Android SM5 device profiles to add vendor-specific defaults

diff --git a/Engine/Config/BaseDeviceProfiles.ini b/Engine/Config/BaseDeviceProfiles.ini
--- a/Engine/Config/BaseDeviceProfiles.ini
+++ b/Engine/Config/BaseDeviceProfiles.ini
@@ -124,6 +124,9 @@
 +DeviceProfileNameAndTypes=Android_Xclipse_xxx,Android
 +DeviceProfileNameAndTypes=Android_Xclipse_xxx_Vulkan,Android
 +DeviceProfileNameAndTypes=Android_Vulkan_SM5,Android
++DeviceProfileNameAndTypes=Android_Adreno_Vulkan_SM5,Android
++DeviceProfileNameAndTypes=Android_Mali_Vulkan_SM5,Android
++DeviceProfileNameAndTypes=Android_Xclipse_Vulkan_SM5,Android
 +DeviceProfileNameAndTypes=Android_PowerVR_G6xxx,Android
 +DeviceProfileNameAndTypes=Android_PowerVR_GT7xxx,Android
 +DeviceProfileNameAndTypes=Android_PowerVR_GE8xxx,Android
@@ -709,14 +712,14 @@ BaseProfileName=VisionPro
 +MatchProfile=(Profile="Android_Adreno4xx",Match=((SourceType=SRC_GpuFamily,CompareType=CMP_Regex,MatchString="Adreno \\(TM\\) 4[0-9][0-9]")))
 +MatchProfile=(Profile="Android_Adreno5xx_Low",Match=((SourceType=SRC_GpuFamily,CompareType=CMP_Regex,MatchString="Adreno \\(TM\\) 5[0-1][0-9]")))
 +MatchProfile=(Profile="Android_Adreno5xx",Match=((SourceType=SRC_GpuFamily,CompareType=CMP_Regex,MatchString="Adreno \\(TM\\) 5[0-9][0-9]")))
-; Example profiles for Vulkan SM5
-;+MatchProfile=(Profile="Android_Vulkan_SM5",Match=((SourceType=SRC_GpuFamily,CompareType=CMP_Regex,MatchString="Adreno \\(TM\\) 6[0-9][0-9]"),(SourceType=SRC_AndroidVersion, CompareType=CMP_Regex,MatchString="([0-9]+).*"),(SourceType=SRC_PreviousRegexMatch,CompareType=CMP_GreaterEqual,MatchString="9"),(SourceType=SRC_SM5Available,CompareType=CMP_Equal,MatchString="true")))
-+MatchProfile=(Profile="Android_Vulkan_SM5",Match=((SourceType=SRC_GpuFamily,CompareType=CMP_Regex,MatchString="Adreno \\(TM\\) 7[0-9][0-9]"),(SourceType=SRC_AndroidVersion, CompareType=CMP_Regex,MatchString="([0-9]+).*"),(SourceType=SRC_PreviousRegexMatch,CompareType=CMP_GreaterEqual,MatchString="9"),(SourceType=SRC_SM5Available,CompareType=CMP_Equal,MatchString="true")))
+; Enable SM5 on Adreno 7xx when SM5 is enabled
++MatchProfile=(Profile="Android_Adreno_Vulkan_SM5",Match=((SourceType=SRC_GpuFamily,CompareType=CMP_Regex,MatchString="Adreno \\(TM\\) 7[0-9][0-9]"),(SourceType=SRC_AndroidVersion, CompareType=CMP_Regex,MatchString="([0-9]+).*"),(SourceType=SRC_PreviousRegexMatch,CompareType=CMP_GreaterEqual,MatchString="10"),(SourceType=SRC_SM5Available,CompareType=CMP_Equal,MatchString="true")))
 ; disable Vulkan on Adreno 660 devices with driver version <= 525
 +MatchProfile=(Profile="Android_Adreno6xx",Match=((SourceType=SRC_GpuFamily,CompareType=CMP_Regex,MatchString="Adreno \\(TM\\) 660"),(SourceType=SRC_GLVersion, CompareType=CMP_Regex,MatchString="V@([0-9]+)"),(SourceType=SRC_PreviousRegexMatch,CompareType=CMP_LessEqual,MatchString="525")))
 ; Enable Vulkan on Adreno 6xx devices with Android >= 10 and driver version >= 400
 +MatchProfile=(Profile="Android_Adreno6xx_Vulkan",Match=((SourceType=SRC_GpuFamily,CompareType=CMP_Regex,MatchString="Adreno \\(TM\\) 6[0-9][0-9]"),(SourceType=SRC_AndroidVersion, CompareType=CMP_Regex,MatchString="([0-9]+).*"),(SourceType=SRC_PreviousRegexMatch,CompareType=CMP_GreaterEqual,MatchString="10"),(SourceType=SRC_VulkanAvailable,CompareType=CMP_Equal,MatchString="true"),(SourceType=SRC_GLVersion, CompareType=CMP_Regex,MatchString="V@([0-9]+)"),(SourceType=SRC_PreviousRegexMatch,CompareType=CMP_GreaterEqual,MatchString="400")))
 +MatchProfile=(Profile="Android_Adreno6xx",Match=((SourceType=SRC_GpuFamily,CompareType=CMP_Regex,MatchString="Adreno \\(TM\\) 6[0-9][0-9]")))
+; Enable Vulkan on Adreno 7xx devices with Android >= 10
 +MatchProfile=(Profile="Android_Adreno7xx_Vulkan",Match=((SourceType=SRC_GpuFamily,CompareType=CMP_Regex,MatchString="Adreno \\(TM\\) 7[0-9][0-9]"),(SourceType=SRC_AndroidVersion, CompareType=CMP_Regex,MatchString="([0-9]+).*"),(SourceType=SRC_PreviousRegexMatch,CompareType=CMP_GreaterEqual,MatchString="10"),(SourceType=SRC_VulkanAvailable,CompareType=CMP_Equal,MatchString="true")))
 +MatchProfile=(Profile="Android_Adreno7xx",Match=((SourceType=SRC_GpuFamily,CompareType=CMP_Regex,MatchString="Adreno \\(TM\\) 7[0-9][0-9]")))
 ; Enable Vulkan on any unknown Adreno running Android 10 or later
@@ -737,7 +740,7 @@ BaseProfileName=VisionPro
 +MatchProfile=(Profile="Android_Mali_G78",Match=((SourceType=SRC_GpuFamily,CompareType=CMP_Regex,MatchString="^Mali\\-G78($| )")))
 +MatchProfile=(Profile="Android_Mali_G710_Vulkan",Match=((SourceType=SRC_GpuFamily,CompareType=CMP_Regex,MatchString="^Mali\\-G710"),(SourceType=SRC_VulkanAvailable,CompareType=CMP_Equal,MatchString="true")))
 +MatchProfile=(Profile="Android_Mali_G710",Match=((SourceType=SRC_GpuFamily,CompareType=CMP_Regex,MatchString="^Mali\\-G710")))
-+MatchProfile=(Profile="Android_Vulkan_SM5_Mali",Match=((SourceType=SRC_GpuFamily,CompareType=CMP_Regex,MatchString="^Mali\\-G7[1-9][0-9]"),(SourceType=SRC_SM5Available,CompareType=CMP_Equal,MatchString="true")))
++MatchProfile=(Profile="Android_Mali_Vulkan_SM5",Match=((SourceType=SRC_GpuFamily,CompareType=CMP_Regex,MatchString="^Mali\\-G7[1-9][0-9]"),(SourceType=SRC_SM5Available,CompareType=CMP_Equal,MatchString="true")))
 +MatchProfile=(Profile="Android_Mali_G7xx_Vulkan",Match=((SourceType=SRC_GpuFamily,CompareType=CMP_Regex,MatchString="^Mali\\-G7[0-9][0-9]"),(SourceType=SRC_VulkanAvailable,CompareType=CMP_Equal,MatchString="true")))
 +MatchProfile=(Profile="Android_Mali_G7xx",Match=((SourceType=SRC_GpuFamily,CompareType=CMP_Regex,MatchString="^Mali\\-G7[0-9][0-9]")))
 
@@ -745,9 +748,9 @@ BaseProfileName=VisionPro
 +MatchProfile=(Profile="Android_Unknown_Vulkan",Match=((SourceType=SRC_GpuFamily,CompareType=CMP_Regex,MatchString="^Mali"),(SourceType=SRC_AndroidVersion, CompareType=CMP_Regex,MatchString="([0-9]+).*"),(SourceType=SRC_PreviousRegexMatch,CompareType=CMP_GreaterEqual,MatchString="10"),(SourceType=SRC_VulkanAvailable,CompareType=CMP_Equal,MatchString="true")))
 
 ; Samsung Exynos / Xclipse GPU-based devices use Vulkan by default.
-+MatchProfile=(Profile="Android_Vulkan_SM5",Match=((SourceType=SRC_GpuFamily,CompareType=CMP_Regex,MatchString="Samsung Xclipse 940"),(SourceType=SRC_SM5Available,CompareType=CMP_Equal,MatchString="true")))
 +MatchProfile=(Profile="Android_Xclipse_5xx_Vulkan",Match=((SourceType=SRC_GpuFamily,CompareType=CMP_Regex,MatchString="Samsung Xclipse 5[0-9][0-9]"),(SourceType=SRC_VulkanAvailable,CompareType=CMP_Equal,MatchString="true")))
 +MatchProfile=(Profile="Android_Xclipse_5xx",Match=((SourceType=SRC_GpuFamily,CompareType=CMP_Regex,MatchString="Samsung Xclipse 5[0-9][0-9]")))
++MatchProfile=(Profile="Android_Xclipse_Vulkan_SM5",Match=((SourceType=SRC_GpuFamily,CompareType=CMP_Regex,MatchString="Samsung Xclipse 9[0-9][0-9]"),(SourceType=SRC_SM5Available,CompareType=CMP_Equal,MatchString="true")))
 +MatchProfile=(Profile="Android_Xclipse_9xx_Vulkan",Match=((SourceType=SRC_GpuFamily,CompareType=CMP_Regex,MatchString="Samsung Xclipse 9[0-9][0-9]"),(SourceType=SRC_VulkanAvailable,CompareType=CMP_Equal,MatchString="true")))
 +MatchProfile=(Profile="Android_Xclipse_9xx",Match=((SourceType=SRC_GpuFamily,CompareType=CMP_Regex,MatchString="Samsung Xclipse 9[0-9][0-9]")))
 +MatchProfile=(Profile="Android_Xclipse_xxx_Vulkan",Match=((SourceType=SRC_GpuFamily,CompareType=CMP_Regex,MatchString="Samsung Xclipse [0-9][0-9][0-9]"),(SourceType=SRC_VulkanAvailable,CompareType=CMP_Equal,MatchString="true")))
@@ -851,7 +854,11 @@ DeviceType=Android
 BaseProfileName=Android_High
 +CVars=r.Android.DisableVulkanSupport=0
 
-; Example base settings for Vulkan SM5 on Android
+[Android_Default DeviceProfile]
+DeviceType=Android
+BaseProfileName=Android_Mid
+
+; Vulkan SM5 on Android
 [Android_Vulkan_SM5 DeviceProfile]
 DeviceType=Android
 BaseProfileName=Android
@@ -888,15 +895,18 @@ BaseProfileName=Android
 +CVars=r.Vulkan.RayTracing.AllowCompaction=0
 +CVars=r.Vulkan.RayTracing.TLASPreferFastTraceTLAS=0
 
+[Android_Adreno_Vulkan_SM5 DeviceProfile]
+DeviceType=Android
+BaseProfileName=Android_Vulkan_SM5
 
-[Android_Vulkan_SM5_Mali DeviceProfile]
+[Android_Mali_Vulkan_SM5 DeviceProfile]
 DeviceType=Android
 BaseProfileName=Android_Vulkan_SM5
 +CVars=r.Vulkan.DepthStencilForceStorageBit=1
 
-[Android_Default DeviceProfile]
+[Android_Xclipse_Vulkan_SM5 DeviceProfile]
 DeviceType=Android
-BaseProfileName=Android_Mid
+BaseProfileName=Android_Vulkan_SM5
 
 ;----------------
 ; Adreno devices, Adreno 4xx and above are supported
