Cooker: Incremental: Update IterativeClassDenyList for indeterminism: remove types with fixes and add on types with new issues.
#rnx

diff --git a/Engine/Config/BaseEditor.ini b/Engine/Config/BaseEditor.ini
--- a/Engine/Config/BaseEditor.ini
+++ b/Engine/Config/BaseEditor.ini
@@ -487,16 +487,14 @@ SaveUnversioned=true
 +IterativeClassScriptPackageAllowList=Allow,<EngineRoot>
 
 +IterativeClassDenyList=/Script/ComputeFramework.ComputeGraph				; UE-203846
-+IterativeClassDenyList=/Script/Engine.AnimSequence 						; UE-202607
++IterativeClassDenyList=/Script/CustomizableObject.CustomizableObject		; UE-209204
 +IterativeClassDenyList=/Script/Engine.BlueprintGeneratedClass 				; Not yet documented
 +IterativeClassDenyList=/Script/Engine.BodySetup 							; UE-202602
 +IterativeClassDenyList=/Script/Engine.MaterialFunction 					; UE-198712
 +IterativeClassDenyList=/Script/Engine.MaterialInterface 					; UE-198712, UE-203846
 +IterativeClassDenyList=/Script/Engine.World 								; Not yet documented
-+IterativeClassDenyList=/Script/CustomizableObject.CustomizableObject		; UE-202596
-+IterativeClassDenyList=/Script/MovieScene.MovieSceneCompiledData			; UE-202603
-+IterativeClassDenyList=/Script/Niagara.NiagaraSystem						; UE-201685
-+IterativeClassDenyList=/Script/StateTreeModule.StateTree					; UE-202608
++IterativeClassDenyList=/Script/PoseSearch.PoseSearchDatabase				; UE-209084
++IterativeClassDenyList=/Script/RigVM.RigVM                                 ; UE-209204
 
 [CookPlatformDataCacheSettings]
 ; Add <ShortClassName>=<Integer> here to limit the number of instances that can have BeginCacheForCookedPlatformData called at once
