Stripping non-iris net metrics when running Iris NetDrivers. Previously all these were set to 0 and output every frame, so this should reduce raw csv size quite a bit.
#tests Tested against Dedicated Server Win64 Cooked Data

diff --git a/Engine/Config/BaseEngine.ini b/Engine/Config/BaseEngine.ini
--- a/Engine/Config/BaseEngine.ini
+++ b/Engine/Config/BaseEngine.ini
@@ -3523,27 +3523,27 @@ IntervalSeconds=1
 
 [/Script/Engine.NetworkMetricsConfig]
 +Listeners=(MetricName=Connections, Class=/Script/Engine.NetworkMetricsCSV_Replication)
-+Listeners=(MetricName=GatherPrioritizeTimeMS, Class=/Script/Engine.NetworkMetricsCSV_Replication)
++Listeners=(MetricName=GatherPrioritizeTimeMS, Class=/Script/Engine.NetworkMetricsCSV_Replication, EnableMode=ENetworkMetricEnableMode::EnableForNonIrisOnly)
 +Listeners=(MetricName=InPacketsClientAvg, Class=/Script/Engine.NetworkMetricsCSV_Replication)
 +Listeners=(MetricName=InPacketsClientMax, Class=/Script/Engine.NetworkMetricsCSV_Replication)
 +Listeners=(MetricName=NumberOfActiveActors, Class=/Script/Engine.NetworkMetricsCSV_Replication)
 +Listeners=(MetricName=NumberOfFullyDormantActors, Class=/Script/Engine.NetworkMetricsCSV_Replication)
 +Listeners=(MetricName=NumClientUpdateLevelVisibility, Class=/Script/Engine.NetworkMetricsCSV_Replication)
-+Listeners=(MetricName=NumOpenChannels, Class=/Script/Engine.NetworkMetricsCSV_Replication)
++Listeners=(MetricName=NumOpenChannels, Class=/Script/Engine.NetworkMetricsCSV_Replication, EnableMode=ENetworkMetricEnableMode::EnableForNonIrisOnly)
 +Listeners=(MetricName=NumReplicateActorCallsPerConAvg, Class=/Script/Engine.NetworkMetricsCSV_Replication)
-+Listeners=(MetricName=NumSkippedObjectEmptyUpdates, Class=/Script/Engine.NetworkMetricsCSV_Replication)
-+Listeners=(MetricName=NumTickingChannels, Class=/Script/Engine.NetworkMetricsCSV_Replication)
++Listeners=(MetricName=NumSkippedObjectEmptyUpdates, Class=/Script/Engine.NetworkMetricsCSV_Replication, EnableMode=ENetworkMetricEnableMode::EnableForNonIrisOnly)
++Listeners=(MetricName=NumTickingChannels, Class=/Script/Engine.NetworkMetricsCSV_Replication, EnableMode=ENetworkMetricEnableMode::EnableForNonIrisOnly)
 +Listeners=(MetricName=OutKBytes, Class=/Script/Engine.NetworkMetricsCSV_Replication)
 +Listeners=(MetricName=OutNetGUIDKBytesSec, Class=/Script/Engine.NetworkMetricsCSV_Replication)
 +Listeners=(MetricName=OutPacketsClientAvg, Class=/Script/Engine.NetworkMetricsCSV_Replication)
 +Listeners=(MetricName=OutPacketsClientMax, Class=/Script/Engine.NetworkMetricsCSV_Replication)
-+Listeners=(MetricName=ReplicateActorTimeMS, Class=/Script/Engine.NetworkMetricsCSV_Replication)
++Listeners=(MetricName=ReplicateActorTimeMS, Class=/Script/Engine.NetworkMetricsCSV_Replication, EnableMode=ENetworkMetricEnableMode::EnableForNonIrisOnly)
 +Listeners=(MetricName=SatConnections, Class=/Script/Engine.NetworkMetricsCSV_Replication)
 +Listeners=(MetricName=ServerReplicateActorTimeMS, Class=/Script/Engine.NetworkMetricsCSV_Replication)
-+Listeners=(MetricName=SharedSerializationPropertyHit, Class=/Script/Engine.NetworkMetricsCSV_Replication)
-+Listeners=(MetricName=SharedSerializationPropertyMiss, Class=/Script/Engine.NetworkMetricsCSV_Replication)
-+Listeners=(MetricName=SharedSerializationRPCHit, Class=/Script/Engine.NetworkMetricsCSV_Replication)
-+Listeners=(MetricName=SharedSerializationRPCMiss, Class=/Script/Engine.NetworkMetricsCSV_Replication)
++Listeners=(MetricName=SharedSerializationPropertyHit, Class=/Script/Engine.NetworkMetricsCSV_Replication, EnableMode=ENetworkMetricEnableMode::EnableForNonIrisOnly)
++Listeners=(MetricName=SharedSerializationPropertyMiss, Class=/Script/Engine.NetworkMetricsCSV_Replication, EnableMode=ENetworkMetricEnableMode::EnableForNonIrisOnly)
++Listeners=(MetricName=SharedSerializationRPCHit, Class=/Script/Engine.NetworkMetricsCSV_Replication, EnableMode=ENetworkMetricEnableMode::EnableForNonIrisOnly)
++Listeners=(MetricName=SharedSerializationRPCMiss, Class=/Script/Engine.NetworkMetricsCSV_Replication, EnableMode=ENetworkMetricEnableMode::EnableForNonIrisOnly)
 +Listeners=(MetricName=OutgoingReliableMessageQueueMaxSize, Class=/Script/Engine.NetworkMetricsCSV_Replication)
 +Listeners=(MetricName=IncomingReliableMessageQueueMaxSize, Class=/Script/Engine.NetworkMetricsCSV_Replication)
 +Listeners=(MetricName=AvgPing, Class=/Script/Engine.NetworkMetricsPerfCounters)
diff --git a/Engine/Source/Runtime/Engine/Private/NetDriver.cpp b/Engine/Source/Runtime/Engine/Private/NetDriver.cpp
--- a/Engine/Source/Runtime/Engine/Private/NetDriver.cpp
+++ b/Engine/Source/Runtime/Engine/Private/NetDriver.cpp
@@ -1807,6 +1807,13 @@ void UNetDriver::SetupNetworkMetricsListeners()
 
 	for (const FNetworkMetricConfig& MetricConfig : Config->Listeners)
 	{
+		// Allow config file to selectively enable/disable listeners
+		if ((MetricConfig.EnableMode == ENetworkMetricEnableMode::EnableForIrisOnly && !IsUsingIrisReplication())
+			|| (MetricConfig.EnableMode == ENetworkMetricEnableMode::EnableForNonIrisOnly && IsUsingIrisReplication()))
+		{
+			continue;
+		}
+
 		FName ListenerClassName = *MetricConfig.Class.ToString();
 		TObjectPtr<UNetworkMetricsBaseListener>* Listener = NetworkMetricsListeners.Find(ListenerClassName);
 
diff --git a/Engine/Source/Runtime/Engine/Public/Net/NetworkMetricsConfig.h b/Engine/Source/Runtime/Engine/Public/Net/NetworkMetricsConfig.h
--- a/Engine/Source/Runtime/Engine/Public/Net/NetworkMetricsConfig.h
+++ b/Engine/Source/Runtime/Engine/Public/Net/NetworkMetricsConfig.h
@@ -4,6 +4,14 @@
 #include "CoreMinimal.h"
 #include "NetworkMetricsConfig.generated.h"
 
+UENUM()
+enum class ENetworkMetricEnableMode : uint8
+{
+	EnableForAllReplication,
+	EnableForIrisOnly,
+	EnableForNonIrisOnly
+};
+
 USTRUCT()
 struct FNetworkMetricConfig
 {
@@ -17,6 +25,9 @@ public:
 	/** A sub-class of UNetworkMetricBaseListener. */
 	UPROPERTY()
 	TSoftClassPtr<class UNetworkMetricsBaseListener> Class;
+
+	UPROPERTY()
+	ENetworkMetricEnableMode EnableMode = ENetworkMetricEnableMode::EnableForAllReplication;
 };
 
 UCLASS(Config=Engine)
