Allow ray tracing with SM5 on Android

CVar r.RayTracing.RequireSM6 is new in 5.4

diff --git a/Engine/Config/BaseDeviceProfiles.ini b/Engine/Config/BaseDeviceProfiles.ini
--- a/Engine/Config/BaseDeviceProfiles.ini
+++ b/Engine/Config/BaseDeviceProfiles.ini
@@ -892,6 +892,7 @@ BaseProfileName=Android
 +CVars=r.DistanceFields=1
 +CVars=r.Vulkan.RayTracing.AllowCompaction=0
 +CVars=r.Vulkan.RayTracing.TLASPreferFastTraceTLAS=0
++CVars=r.RayTracing.RequireSM6=0
 
 [Android_Adreno_Vulkan_SM5 DeviceProfile]
 DeviceType=Android
