Report the NumConnections network stat to PerfCounters and CSV.

diff --git a/Engine/Config/BaseEngine.ini b/Engine/Config/BaseEngine.ini
--- a/Engine/Config/BaseEngine.ini
+++ b/Engine/Config/BaseEngine.ini
@@ -3522,6 +3522,7 @@ IntervalSeconds=0
 IntervalSeconds=1
 
 [/Script/Engine.NetworkMetricsConfig]
++Listeners=(MetricName=NumConnections, Class=/Script/Engine.NetworkMetricsCSV_Replication)
 +Listeners=(MetricName=Connections, Class=/Script/Engine.NetworkMetricsCSV_Replication)
 +Listeners=(MetricName=GatherPrioritizeTimeMS, Class=/Script/Engine.NetworkMetricsCSV_Replication, EnableMode=ENetworkMetricEnableMode::EnableForNonIrisOnly)
 +Listeners=(MetricName=InPacketsClientAvg, Class=/Script/Engine.NetworkMetricsCSV_Replication)
@@ -3546,6 +3547,7 @@ IntervalSeconds=1
 +Listeners=(MetricName=SharedSerializationRPCMiss, Class=/Script/Engine.NetworkMetricsCSV_Replication, EnableMode=ENetworkMetricEnableMode::EnableForNonIrisOnly)
 +Listeners=(MetricName=OutgoingReliableMessageQueueMaxSize, Class=/Script/Engine.NetworkMetricsCSV_Replication)
 +Listeners=(MetricName=IncomingReliableMessageQueueMaxSize, Class=/Script/Engine.NetworkMetricsCSV_Replication)
++Listeners=(MetricName=NumConnections, Class=/Script/Engine.NetworkMetricsPerfCounters)
 +Listeners=(MetricName=AvgPing, Class=/Script/Engine.NetworkMetricsPerfCounters)
 +Listeners=(MetricName=InBunches, Class=/Script/Engine.NetworkMetricsPerfCounters)
 +Listeners=(MetricName=InPackets, Class=/Script/Engine.NetworkMetricsPerfCounters)
