Re-enable ARKit source autocreation by default in the editor.

diff --git a/Engine/Config/BaseEngine.ini b/Engine/Config/BaseEngine.ini
--- a/Engine/Config/BaseEngine.ini
+++ b/Engine/Config/BaseEngine.ini
@@ -3428,7 +3428,7 @@ PakSigningRequired=True
 PakEncryptionRequired=True
 
 [/Script/AppleARKit.AppleARKitSettings]
-bEnableLiveLinkForFaceTracking=false
+bEnableLiveLinkForFaceTracking=true
 LiveLinkPublishingPort=11111
 ; by default arkit will be added to UIRequiredDeviceCapabilities and devices must support
 ; arkit. If false games must check that arkit is supported before enabling arkit features
