MPCook: CookPackageSplitter: Toggle enable retraction of generated packages.
#rnx

diff --git a/Engine/Config/BaseEditor.ini b/Engine/Config/BaseEditor.ini
--- a/Engine/Config/BaseEditor.ini
+++ b/Engine/Config/BaseEditor.ini
@@ -398,7 +398,7 @@ CookProcessCountFromEditor=-1
 ; Debug setting to specify that the packages generated by a generator package must be saved on the same cookworker
 ; that saved the generator package. This is bad for performance, but is provided in case there is a bug with moving
 ; them to other workers. Values = AnyWorker, AllOnSameWorker, SomeOnSameWorker, NoneOnSameWorker
-MPCookGeneratorSplit=AllOnSameWorker
+MPCookGeneratorSplit=AnyWorker
 
 +ExcludedIniParams=Game:/Script/BlueprintContext.McpContext
 +ExcludedIniParams=Editor:CookSettings
