[Backout] - CL32546429, it is likely the cause of hanging autocooks.

Original CL Desc
-----------------------------------------------------------------
MPCook: CookPackageSplitter: Toggle enable retraction of generated packages.

diff --git a/Engine/Config/BaseEditor.ini b/Engine/Config/BaseEditor.ini
--- a/Engine/Config/BaseEditor.ini
+++ b/Engine/Config/BaseEditor.ini
@@ -398,7 +398,7 @@ CookProcessCountFromEditor=-1
 ; Debug setting to specify that the packages generated by a generator package must be saved on the same cookworker
 ; that saved the generator package. This is bad for performance, but is provided in case there is a bug with moving
 ; them to other workers. Values = AnyWorker, AllOnSameWorker, SomeOnSameWorker, NoneOnSameWorker
-MPCookGeneratorSplit=AnyWorker
+MPCookGeneratorSplit=AllOnSameWorker
 
 +ExcludedIniParams=Game:/Script/BlueprintContext.McpContext
 +ExcludedIniParams=Editor:CookSettings
