Update minimum GPU driver versions for 5.4:
* 536.xx for Nvidia
* 24.3.1 for AMD
* 101.5333 for Intel ARC and newer

Legacy GPUs are allowed to use older drivers in DX11. Also updated these versions to the latest available for each vendor's legacy support line.

#rnx

diff --git a/Engine/Config/BaseHardware.ini b/Engine/Config/BaseHardware.ini
--- a/Engine/Config/BaseHardware.ini
+++ b/Engine/Config/BaseHardware.ini
@@ -33,56 +33,60 @@
 ;
 ;-----------------------------------------------------------------------------------------------------------------
 
-
-
 [GPU_NVIDIA Windows]
-SuggestedDriverVersion="527.27"
+SuggestedDriverVersion="536.40"
 
-; To prevent problems with older drivers.
-+DriverDenyList=(DriverVersion="<473.47", Reason="These driver versions have known stability issues and missing features", RHI="D3D11")
-+DriverDenyList=(DriverVersion="<527.27", Reason="These driver versions have known stability issues and missing features", RHI="D3D12")
-+DriverDenyList=(DriverVersion="<527.27", Reason="These driver versions have known stability issues and missing features", RHI="Vulkan")
+; The newest driver which supports legacy Nvidia GPUs (700 series and older) is 474.64. We allow this driver for DX11.
++DriverDenyList=(DriverVersion="<474.64", Reason="These driver versions have known stability issues and missing features", RHI="D3D11")
+
+; DX12 and Vulkan require recent drivers. The 536 series was released in June 2023.
++DriverDenyList=(DriverVersion="<536.00", Reason="These driver versions have known stability issues and missing features", RHI="D3D12")
++DriverDenyList=(DriverVersion="<536.00", Reason="These driver versions have known stability issues and missing features", RHI="Vulkan")
 
 ;-----------------------------------------------------------------------------------------------------------------
 
 [GPU_NVIDIA Linux]
-SuggestedDriverVersion="515.48"
-
-; UE5 in Linux/Vulkan requires recent drivers for Nanite/Lumen and has other stability issues with older drivers.
-+DriverDenyList=(DriverVersion="<515.48", Reason="These driver versions have known stability issues and missing features in Vulkan")
-
-
+SuggestedDriverVersion="535.86"
 
+; The 535 series of the Linux driver was released in June 2023.
++DriverDenyList=(DriverVersion="<535.00", Reason="These driver versions have known stability issues and missing features")
 
-;-----------------------------------------------------------------------------------------------------------------
 ;-----------------------------------------------------------------------------------------------------------------
 
 [GPU_AMD Windows]
-SuggestedDriverVersion="23.3.2"
+SuggestedDriverVersion="24.3.1"
+
+; The newest driver which supports legacy AMD GPUs (GCN1-GCN3 chips like the Radeon R9 or older) is 22.6.1. We allow this driver for DX11.
++DriverDenyList=(DriverVersion="<27.20.20913.2000", Reason="These driver versions have known stability issues and missing features", RHI="D3D11")
 
-; To prevent problems with older drivers.
-+DriverDenyList=(DriverVersion="<30.0.14011.2006", Reason="These driver versions have known stability issues and missing features", RHI="D3D11")
-+DriverDenyList=(DriverVersion="<30.0.14011.2006", Reason="These driver versions have known stability issues and missing features", RHI="D3D12")
-; 23.3.2 is necessary to prevent pipeline compilation crash in Vulkan
-+DriverDenyList=(DriverVersion="<31.0.14043.1000", Reason="These driver versions have known stability issues and missing features", RHI="Vulkan")
+; The minimum version for DX12 is Vulkan is 24.3.1, published in March 2024. There are separate packages for different GPU generations and they
+; have different internal versions, so we can't check for a version number explicitly (the RDNA driver is 31.0.24027.1012, older GPUs have
+; 31.0.21912.14). We'll use a driver date check instead, as both packages have internal date stamps newer than Feb 19th.
++DriverDenyList=(DriverDate="<2-19-2024", Reason="These driver versions have known stability issues and missing features", RHI="D3D12")
++DriverDenyList=(DriverDate="<2-19-2024", Reason="These driver versions have known stability issues and missing features", RHI="Vulkan")
 
 ;-----------------------------------------------------------------------------------------------------------------
 
 [GPU_AMD Linux]
-SuggestedDriverVersion="21.2.6"
+SuggestedDriverVersion="23.3.1"
 
-; UE5 in Linux/Vulkan requires recent drivers for Nanite/Lumen and has other stability issues with older drivers (need >= RADV 21.0.3)
-+DriverDenyList=(DriverVersion="<21.2.6", Reason="These driver versions have known stability issues and missing features in Vulkan")
+; RADV builds older than 23.3.1 have issues running UE5, so we'll use this as the minimum version (released December 2023). Note that it's also
+; possible to use the AMD proprietary driver on Linux, and we don't distinguish between driver vendors. The current version for the AMD driver as
+; of writing this is 23.40.2, so 23.3.1 is a valid version for that package as well (although a fairly old one).
++DriverDenyList=(DriverVersion="<23.3.1", Reason="These driver versions have known stability issues and missing features")
 
-
-
-
-;-----------------------------------------------------------------------------------------------------------------
 ;-----------------------------------------------------------------------------------------------------------------
 
 [GPU_Intel Windows]
-SuggestedDriverVersion="30.0.101.1340"
-+DriverDenyList=(DriverVersion="<101.1191", Reason="These driver versions have known stability issues and missing features")
+SuggestedDriverVersion="101.5333"
+
+; Legacy integrated GPUs are supported by several driver packages (depending on CPU generation) and they all have different version numbers, so
+; we will check the date instead. The latest driver for Ivy Bridge chips is 105.161, dated Aug 6th 2020, so we will allow that for D3D11.
++DriverDenyList=(DriverDate="<8-6-2020", Reason="These driver versions have known stability issues and missing features", RHI="D3D11")
+
+; DX12 and Vulkan should not be used on integrated GPUs. ARC and newer chips require recent drivers. 101.5333 was released in February 2024.
++DriverDenyList=(DriverVersion="<101.5333", Reason="These driver versions have known stability issues and missing features", RHI="D3D12")
++DriverDenyList=(DriverVersion="<101.5333", Reason="These driver versions have known stability issues and missing features", RHI="Vulkan")
 
 ;-----------------------------------------------------------------------------------------------------------------
 
