- Enable GPU particles & sorting for Quest 2/3

diff --git a/Engine/Config/BaseDeviceProfiles.ini b/Engine/Config/BaseDeviceProfiles.ini
--- a/Engine/Config/BaseDeviceProfiles.ini
+++ b/Engine/Config/BaseDeviceProfiles.ini
@@ -1107,6 +1107,8 @@ BaseProfileName=Android_High
 [Meta_Quest_3 DeviceProfile]
 DeviceType=Android
 BaseProfileName=Meta_Quest_Pro
++CVars=fx.NiagaraAllowGPUParticles=1
++CVars=FX.AllowGPUSorting=1
 
 [Meta_Quest_Pro DeviceProfile]
 DeviceType=Android
@@ -1116,6 +1118,8 @@ BaseProfileName=Oculus_Quest2
 DeviceType=Android
 BaseProfileName=Oculus_Quest
 +CVars=r.Mobile.Oculus.ForceSymmetric=1
++CVars=fx.NiagaraAllowGPUParticles=1
++CVars=FX.AllowGPUSorting=1
 
 [Oculus_Quest DeviceProfile]
 DeviceType=Android
