Map Template
* Add a way to localize map template from config file.
* Add LocName for Open World and Basic map template.

diff --git a/Engine/Config/BaseEngine.ini b/Engine/Config/BaseEngine.ini
--- a/Engine/Config/BaseEngine.ini
+++ b/Engine/Config/BaseEngine.ini
@@ -1847,8 +1847,8 @@ InEditorGameURLOptions=
 
 [/Script/UnrealEd.UnrealEdEngine]
 AutoSaveIndex=0
-+TemplateMapInfos=(Thumbnail="/Engine/Maps/Templates/Thumbnails/OpenWorld.OpenWorld",Map="/Engine/Maps/Templates/OpenWorld",DisplayName="Open World",Category="OpenWorld")
-+TemplateMapInfos=(Thumbnail="/Engine/Maps/Templates/Thumbnails/Basic.Basic",Map="/Engine/Maps/Templates/Template_Default",DisplayName="Basic")
++TemplateMapInfos=(Thumbnail="/Engine/Maps/Templates/Thumbnails/OpenWorld.OpenWorld",Map="/Engine/Maps/Templates/OpenWorld", LocKeyName="OpenWorld", DisplayName="Open World",Category="OpenWorld")
++TemplateMapInfos=(Thumbnail="/Engine/Maps/Templates/Thumbnails/Basic.Basic",Map="/Engine/Maps/Templates/Template_Default", LocKeyName="Basic", DisplayName="Basic")
 
 +PackagesToBeFullyLoadedAtStartup=/Engine/EditorMaterials/Cross
 +PackagesToBeFullyLoadedAtStartup=/Engine/EditorMaterials/Cross_Mat
diff --git a/Engine/Source/Editor/NewLevelDialog/Private/NewLevelDialogModule.cpp b/Engine/Source/Editor/NewLevelDialog/Private/NewLevelDialogModule.cpp
--- a/Engine/Source/Editor/NewLevelDialog/Private/NewLevelDialogModule.cpp
+++ b/Engine/Source/Editor/NewLevelDialog/Private/NewLevelDialogModule.cpp
@@ -369,7 +369,8 @@ private:
 			TSharedPtr<FNewLevelTemplateItem> Item = MakeShareable(new FNewLevelTemplateItem());
 			Item->TemplateMapInfo = TemplateMapInfo;
 			Item->Type = FNewLevelTemplateItem::NewLevelType::Template;
-			Item->Name = TemplateMapInfo.DisplayName;
+			// This would keep the loc key for QA localization (e.g. culture=keys in console)
+			Item->Name = TemplateMapInfo.LocKeyName.IsEmpty() ? TemplateMapInfo.DisplayName : FInternationalization::ForUseOnlyByLocMacroAndGraphNodeTextLiterals_CreateText(*TemplateMapInfo.DisplayName.ToString(), TEXT("NewLevelDialog"), *TemplateMapInfo.LocKeyName);
 			Item->Category = TemplateMapInfo.Category;
 			Item->OriginalIndex = CurrentIndex++;
 
diff --git a/Engine/Source/Editor/UnrealEd/Classes/Editor/TemplateMapInfo.h b/Engine/Source/Editor/UnrealEd/Classes/Editor/TemplateMapInfo.h
--- a/Engine/Source/Editor/UnrealEd/Classes/Editor/TemplateMapInfo.h
+++ b/Engine/Source/Editor/UnrealEd/Classes/Editor/TemplateMapInfo.h
@@ -26,6 +26,10 @@ struct FTemplateMapInfo
 	UPROPERTY(config, meta = (AllowedClasses = "/Script/Engine.World"))
 	FSoftObjectPath Map;
 
+	/** Optional localization name for this map template  */
+	UPROPERTY(config)
+	FString LocKeyName;
+
 	/** Optional display name override for this map template  */
 	UPROPERTY(config)
 	FText DisplayName;
