Map Template
* Better way to localize (directly in .ini file) the map template.

diff --git a/Engine/Config/BaseEngine.ini b/Engine/Config/BaseEngine.ini
--- a/Engine/Config/BaseEngine.ini
+++ b/Engine/Config/BaseEngine.ini
@@ -1847,8 +1847,8 @@ InEditorGameURLOptions=
 
 [/Script/UnrealEd.UnrealEdEngine]
 AutoSaveIndex=0
-+TemplateMapInfos=(Thumbnail="/Engine/Maps/Templates/Thumbnails/OpenWorld.OpenWorld",Map="/Engine/Maps/Templates/OpenWorld", LocKeyName="OpenWorld", DisplayName="Open World",Category="OpenWorld")
-+TemplateMapInfos=(Thumbnail="/Engine/Maps/Templates/Thumbnails/Basic.Basic",Map="/Engine/Maps/Templates/Template_Default", LocKeyName="Basic", DisplayName="Basic")
++TemplateMapInfos=(Thumbnail="/Engine/Maps/Templates/Thumbnails/OpenWorld.OpenWorld",Map="/Engine/Maps/Templates/OpenWorld", DisplayName=NSLOCTEXT("NewLevelDialog", "OpenWorld", "Open World"), Category="OpenWorld")
++TemplateMapInfos=(Thumbnail="/Engine/Maps/Templates/Thumbnails/Basic.Basic",Map="/Engine/Maps/Templates/Template_Default", DisplayName=NSLOCTEXT("NewLevelDialog", "Basic", "Basic"))
 
 +PackagesToBeFullyLoadedAtStartup=/Engine/EditorMaterials/Cross
 +PackagesToBeFullyLoadedAtStartup=/Engine/EditorMaterials/Cross_Mat
diff --git a/Engine/Source/Editor/NewLevelDialog/Private/NewLevelDialogModule.cpp b/Engine/Source/Editor/NewLevelDialog/Private/NewLevelDialogModule.cpp
--- a/Engine/Source/Editor/NewLevelDialog/Private/NewLevelDialogModule.cpp
+++ b/Engine/Source/Editor/NewLevelDialog/Private/NewLevelDialogModule.cpp
@@ -369,8 +369,7 @@ private:
 			TSharedPtr<FNewLevelTemplateItem> Item = MakeShareable(new FNewLevelTemplateItem());
 			Item->TemplateMapInfo = TemplateMapInfo;
 			Item->Type = FNewLevelTemplateItem::NewLevelType::Template;
-			// This would keep the loc key for QA localization (e.g. culture=keys in console)
-			Item->Name = TemplateMapInfo.LocKeyName.IsEmpty() ? TemplateMapInfo.DisplayName : FInternationalization::ForUseOnlyByLocMacroAndGraphNodeTextLiterals_CreateText(*TemplateMapInfo.DisplayName.ToString(), TEXT("NewLevelDialog"), *TemplateMapInfo.LocKeyName);
+			Item->Name = TemplateMapInfo.DisplayName;
 			Item->Category = TemplateMapInfo.Category;
 			Item->OriginalIndex = CurrentIndex++;
 
diff --git a/Engine/Source/Editor/UnrealEd/Classes/Editor/TemplateMapInfo.h b/Engine/Source/Editor/UnrealEd/Classes/Editor/TemplateMapInfo.h
--- a/Engine/Source/Editor/UnrealEd/Classes/Editor/TemplateMapInfo.h
+++ b/Engine/Source/Editor/UnrealEd/Classes/Editor/TemplateMapInfo.h
@@ -26,10 +26,6 @@ struct FTemplateMapInfo
 	UPROPERTY(config, meta = (AllowedClasses = "/Script/Engine.World"))
 	FSoftObjectPath Map;
 
-	/** Optional localization name for this map template  */
-	UPROPERTY(config)
-	FString LocKeyName;
-
 	/** Optional display name override for this map template  */
 	UPROPERTY(config)
 	FText DisplayName;
