Incremental cook: Remove Material and MaterialFunction from the deny list since their iterative false skip jiras have been fixed. There is still a known issue with MaterialInstanceConstant; ETA for that fix is in a few days.
#rnx

diff --git a/Engine/Config/BaseEditor.ini b/Engine/Config/BaseEditor.ini
--- a/Engine/Config/BaseEditor.ini
+++ b/Engine/Config/BaseEditor.ini
@@ -494,8 +494,6 @@ SaveUnversioned=true
 +IterativeClassDenyList=/Script/CustomizableObject.CustomizableObject		; UE-209204
 +IterativeClassDenyList=/Script/Engine.BlueprintGeneratedClass 				; Not yet documented
 +IterativeClassDenyList=/Script/Engine.BodySetup 							; UE-202602
-+IterativeClassDenyList=/Script/Engine.MaterialFunction 					; UE-198712
-+IterativeClassDenyList=/Script/Engine.MaterialInterface 					; UE-198712, UE-203846
 +IterativeClassDenyList=/Script/Engine.World 								; Not yet documented
 +IterativeClassDenyList=/Script/PoseSearch.PoseSearchDatabase				; UE-209084
 +IterativeClassDenyList=/Script/RigVM.RigVM                                 ; UE-209204
