Allow AFS APK to be passed through
Block AFS token from packaged ini
Rotate the AFS SecurityToken

#android

diff --git a/Engine/Config/BaseGame.ini b/Engine/Config/BaseGame.ini
--- a/Engine/Config/BaseGame.ini
+++ b/Engine/Config/BaseGame.ini
@@ -157,6 +157,7 @@ WriteBackMetadataToAssetRegistry=Disabled
 
 +IniSectionDenylist=HordeStorageServers
 +IniSectionDenylist=StorageServers
++IniSectionDenylist=/Script/AndroidFileServerEditor.AndroidFileServerRuntimeSettings
 
 [/Script/DeveloperToolSettings.PlatformsMenuSettings]
 StagingDirectory=(Path="")
