Improve startup time for CookShadersCommandlet
Add to CommandletsUsingAR to enable async asset registry loading
Add DisablePythonForCommandlet
Built in plugin loading time reduced from 7m to 30s

#tests local testing of commandlet

diff --git a/Engine/Config/BaseEngine.ini b/Engine/Config/BaseEngine.ini
--- a/Engine/Config/BaseEngine.ini
+++ b/Engine/Config/BaseEngine.ini
@@ -2712,6 +2712,9 @@ AspectRatioAxisConstraint=AspectRatio_MaintainYFOV
 +Class=ParticleSystem
 +Class=Texture2D
 
+[PythonScriptPlugin]
++DisablePythonForCommandlet=CookShaders
+
 [AssetRegistry]
 ; Default Path Filters not scanned by the asset registry when searching for all assets
 ; These are long package paths
@@ -2758,6 +2761,7 @@ bFilterSearchableNames=true
 +CommandletsUsingAR=GatherText
 +CommandletsUsingAR=GenerateNaniteDisplacedMesh
 +CommandletsUsingAR=WorldPartitionBuilder
++CommandletsUsingAR=CookShaders
 
 ; AssetDataGathering can be cached in the editor; there are two different caches, one for the directory scan to find all
 ; packages and one for tags and other data read from the packages. The directory scan cache by default is only read
