HWRT: Set TLAS scalability such as FastBuild is used in High and FastTrace in Epic.
* The reason is that Epic mode will trace many more rays and will benefit from better TLAS build while High mode is mostly dominated by TLAS build cost.

diff --git a/Engine/Config/BaseScalability.ini b/Engine/Config/BaseScalability.ini
--- a/Engine/Config/BaseScalability.ini
+++ b/Engine/Config/BaseScalability.ini
@@ -259,6 +259,7 @@ r.Shadow.Virtual.SMRT.SamplesPerRayLocal=8
 r.DistanceFieldAO=0
 r.Lumen.DiffuseIndirect.Allow=0
 r.SkyLight.RealTimeReflectionCapture=0
+r.RayTracing.Scene.BuildMode=0
 
 [GlobalIlluminationQuality@1]
 ; Enable DFAO when Lumen is off
@@ -266,6 +267,7 @@ r.DistanceFieldAO=1
 r.AOQuality=1
 r.Lumen.DiffuseIndirect.Allow=0
 r.SkyLight.RealTimeReflectionCapture=0
+r.RayTracing.Scene.BuildMode=0
 
 [GlobalIlluminationQuality@2]
 r.DistanceFieldAO=1
@@ -289,6 +291,7 @@ r.Lumen.TranslucencyVolume.TraceFromVolume=0
 r.Lumen.TranslucencyVolume.RadianceCache.ProbeResolution=8
 r.Lumen.TranslucencyVolume.RadianceCache.NumProbesToTraceBudget=200
 r.SkyLight.RealTimeReflectionCapture=1
+r.RayTracing.Scene.BuildMode=0
 
 [GlobalIlluminationQuality@3]
 r.DistanceFieldAO=1
@@ -313,6 +316,7 @@ r.Lumen.TranslucencyVolume.TracingOctahedronResolution=3
 r.Lumen.TranslucencyVolume.RadianceCache.ProbeResolution=8
 r.Lumen.TranslucencyVolume.RadianceCache.NumProbesToTraceBudget=200
 r.SkyLight.RealTimeReflectionCapture=1
+r.RayTracing.Scene.BuildMode=1
 
 [GlobalIlluminationQuality@Cine]
 r.DistanceFieldAO=1
@@ -337,6 +341,7 @@ r.Lumen.TranslucencyVolume.TracingOctahedronResolution=4
 r.Lumen.TranslucencyVolume.RadianceCache.ProbeResolution=16
 r.Lumen.TranslucencyVolume.RadianceCache.NumProbesToTraceBudget=1000
 r.SkyLight.RealTimeReflectionCapture=1
+r.RayTracing.Scene.BuildMode=1
 
 
 ;-----------------------------------------------------------------------------------------------------------------
