Data Driven CVars now use a key driven config array. With the key being Name. This will make it so that no two entries in CVarsArray can have the same entry for Name, ensuring that there are no conflicting cvars

This addresses bugs that could occur if you cvars had the same name by preventing the behavior altogether.

diff --git a/Engine/Config/BaseEngine.ini b/Engine/Config/BaseEngine.ini
--- a/Engine/Config/BaseEngine.ini
+++ b/Engine/Config/BaseEngine.ini
@@ -2297,6 +2297,10 @@ MeURL="https://graph.facebook.com/`ver/me?access_token=`token"
 +LoginDomains=".facebook.com"
 bUsePopup=true
 
+[/Script/Engine.DataDrivenConsoleVariableSettings]
+; @CVarsArray=Name makes it so Name is used as a unique key for the array. Likely should never be overwritten unless there is a unique circumstance 
+@CVarsArray=Name
+
 [OnlineSubsystemFacebook.OnlineSharingFacebook]
 PermissionsURL="https://graph.facebook.com/`ver/me/permissions?access_token=`token"
 
