Increased PersistentAllocatorReserveSizeMB to 1024 for the Engine

diff --git a/Engine/Config/BaseEngine.ini b/Engine/Config/BaseEngine.ini
--- a/Engine/Config/BaseEngine.ini
+++ b/Engine/Config/BaseEngine.ini
@@ -2199,7 +2199,7 @@ bRecompileOnLoad=true
 
 [MemoryPools]
 FLightPrimitiveInteractionInitialBlockSize=512
-PersistentAllocatorReserveSizeMB=512
+PersistentAllocatorReserveSizeMB=1024
 
 [ProcessLimits]
 VirtualMemoryLimitInKB=0
