Clean up errors in .ini config files reported by Rider.

#rnx

diff --git a/Engine/Config/BaseEngine.ini b/Engine/Config/BaseEngine.ini
--- a/Engine/Config/BaseEngine.ini
+++ b/Engine/Config/BaseEngine.ini
@@ -589,7 +589,7 @@ DefaultWorldPartitionHLODLayer=/Engine/Maps/Templates/HLODs/HLODLayer_Instanced.
 +PropertyRedirects=(OldName="Widget.Opacity", NewName="Widget.RenderOpacity")
 +FunctionRedirects=(OldName="Widget.GetOpacity", NewName="Widget.GetRenderOpacity")
 +FunctionRedirects=(OldName="Widget.SetOpacity", NewName="Widget.SetRenderOpacity")
-+EnumRedirects=(OldName="ENetDormancy",ValueChanges=(("DORN_MAX","DORM_MAX"))
++EnumRedirects=(OldName="ENetDormancy",ValueChanges=(("DORN_MAX","DORM_MAX")))
 
 +PackageRedirects=(OldName="/Script/EditorScriptingUtilitiesEditor", NewName="/Script/AssetScriptingUtilitiesEditor")
 
@@ -604,7 +604,7 @@ DefaultWorldPartitionHLODLayer=/Engine/Maps/Templates/HLODs/HLODLayer_Instanced.
 +FunctionRedirects=(OldName="/Script/HeadMountedDisplay.HeadMountedDisplayFunctionLibrary.AddDeviceVisualizationComponent",NewName="/Script/HeadMountedDisplay.XRAssetFunctionLibrary.AddDeviceVisualizationComponentBlocking")
 +FunctionRedirects=(OldName="/Script/HeadMountedDisplay.HeadMountedDisplayFunctionLibrary.AddNamedDeviceVisualizationComponent",NewName="/Script/HeadMountedDisplay.XRAssetFunctionLibrary.AddNamedDeviceVisualizationComponentBlocking")
 +EnumRedirects=(OldName="ESimulationSpace",ValueChanges=(("RootBoneSpace", "BaseBoneSpace")))
-+EnumRedirects=(OldName="EColorVisionDeficiency", NewName="/Script/SlateCore.EColorVisionDeficiency", ValueChanges=(("CVD_NormalVision", "NormalVision"), ("CVD_Deuteranomly", "NormalVision"), ("CVD_Deuteranopia", "Deuteranope"), ("CVD_Protanomly", "NormalVision"), ("CVD_Protanopia", "Protanope"), ("CVD_Tritanomaly", "NormalVision"), ("CVD_Tritanopia", "Tritanope"), ("CVD_Achromatopsia", "NormalVision"))
++EnumRedirects=(OldName="EColorVisionDeficiency", NewName="/Script/SlateCore.EColorVisionDeficiency", ValueChanges=(("CVD_NormalVision", "NormalVision"), ("CVD_Deuteranomly", "NormalVision"), ("CVD_Deuteranopia", "Deuteranope"), ("CVD_Protanomly", "NormalVision"), ("CVD_Protanopia", "Protanope"), ("CVD_Tritanomaly", "NormalVision"), ("CVD_Tritanopia", "Tritanope"), ("CVD_Achromatopsia", "NormalVision")))
 
 +ClassRedirects=(OldName="NavigationSystem",NewName="/Script/NavigationSystem.NavigationSystemV1")
 +ClassRedirects=(OldName="NavMeshBoundsVolume",NewName="/Script/NavigationSystem.NavMeshBoundsVolume")
@@ -803,14 +803,14 @@ DefaultWorldPartitionHLODLayer=/Engine/Maps/Templates/HLODs/HLODLayer_Instanced.
 +ClassRedirects=(OldName="/Script/AnimGraph.AnimGraphNode_SubInput",NewName="/Script/AnimGraph.AnimGraphNode_LinkedInputPose")
 +ClassRedirects=(OldName="/Script/AnimGraph.AnimGraphNode_Layer",NewName="/Script/AnimGraph.AnimGraphNode_LinkedAnimLayer")
 +PropertyRedirects=(OldName="PersonaPreviewSceneDescription.SubInstanceTag",NewName="LinkedAnimGraphTag")
-+EnumRedirects=(OldName="/Script/Engine.EPreviewAnimationBlueprintApplicationMethod",NewName="/Script/Engine.EPreviewAnimationBlueprintApplicationMethod",ValueChanges=(("OverlayLayer", "LinkedLayers"), ("SubInstance", "LinkedAnimGraph"))
++EnumRedirects=(OldName="/Script/Engine.EPreviewAnimationBlueprintApplicationMethod",NewName="/Script/Engine.EPreviewAnimationBlueprintApplicationMethod",ValueChanges=(("OverlayLayer", "LinkedLayers"), ("SubInstance", "LinkedAnimGraph")))
 +PropertyRedirects=(OldName="AnimClassData.SubInstanceNodeProperties",NewName="LinkedAnimGraphNodeProperties")
 +PropertyRedirects=(OldName="AnimClassData.LayerNodeProperties",NewName="LinkedAnimLayerNodeProperties")
 +ClassRedirects=(OldName="/Script/MeshEditingToolset.BaseBrushTool",NewName="/Script/InteractiveToolsFramework.BaseBrushTool")
 +ClassRedirects=(OldName="/Script/MeshEditingToolset.BrushBaseProperties",NewName="/Script/InteractiveToolsFramework.BrushBaseProperties")
 
 ; 4.25
-+EnumRedirects=(OldName="ECurveBlendOption",ValueChanges=(("ECurveBlendOption::MaxWeight", "ECurveBlendOption::Override"))
++EnumRedirects=(OldName="ECurveBlendOption",ValueChanges=(("ECurveBlendOption::MaxWeight", "ECurveBlendOption::Override")))
 +ClassRedirects=(OldName="/Script/OnlineBlueprintSupport.K2Node_LatentOnlineCall", NewName="/Script/BlueprintGraph.K2Node_AsyncAction")
 +ClassRedirects=(OldName="/Script/Kismet.K2Node_AsyncAction", NewName="/Script/BlueprintGraph.K2Node_AsyncAction")
 
diff --git a/Engine/Plugins/Animation/ControlRig/Config/BaseControlRig.ini b/Engine/Plugins/Animation/ControlRig/Config/BaseControlRig.ini
--- a/Engine/Plugins/Animation/ControlRig/Config/BaseControlRig.ini
+++ b/Engine/Plugins/Animation/ControlRig/Config/BaseControlRig.ini
@@ -24,7 +24,7 @@
 
 +ClassRedirects=(OldName="RigVMGraph",NewName="/Script/RigVMDeveloper.RigVMGraph",ValueChanges=(("Model","RigVMModel")))
 
-+EnumRedirects=(OldName="ERigElementType",ValueChanges=(("Space","Null"))
++EnumRedirects=(OldName="ERigElementType",ValueChanges=(("Space","Null")))
 +EnumRedirects=(OldName="/Script/ControlRig.EControlRigRotationOrder",NewName="/Script/AnimationCore.EEulerRotationOrder")
 +EnumRedirects=(OldName="EControlRigRotationOrder",NewName="/Script/AnimationCore.EEulerRotationOrder")
 
diff --git a/Engine/Plugins/Interchange/Runtime/Config/BaseInterchange.ini b/Engine/Plugins/Interchange/Runtime/Config/BaseInterchange.ini
--- a/Engine/Plugins/Interchange/Runtime/Config/BaseInterchange.ini
+++ b/Engine/Plugins/Interchange/Runtime/Config/BaseInterchange.ini
@@ -2,7 +2,7 @@
 bCanSaveAllSections=True
 
 [CoreRedirects]
-+EnumRedirects=(OldName="/Script/InterchangePipelines.EInterchangeMaterialImportOption",ValueChanges=(("DoNotImport","ImportAsMaterials"))
++EnumRedirects=(OldName="/Script/InterchangePipelines.EInterchangeMaterialImportOption",ValueChanges=(("DoNotImport","ImportAsMaterials")))
 
 ; Rename Interchange CineCamera to PhysicalCamera FactoryNode
 +ClassRedirects=(OldName="/Script/InterchangeFactoryNodes.InterchangeCineCameraFactoryNode", NewName="/Script/InterchangeFactoryNodes.InterchangePhysicalCameraFactoryNode")
diff --git a/Engine/Plugins/Runtime/StateTree/Config/BaseStateTree.ini b/Engine/Plugins/Runtime/StateTree/Config/BaseStateTree.ini
--- a/Engine/Plugins/Runtime/StateTree/Config/BaseStateTree.ini
+++ b/Engine/Plugins/Runtime/StateTree/Config/BaseStateTree.ini
@@ -55,7 +55,7 @@
 +PropertyRedirects=(OldName="CompactStateTransition.Event",NewName="CompactStateTransition.Trigger")
 +ClassRedirects=(OldName="/Script/StateTreeModule.StateTreeItemBlueprintBase",NewName="/Script/StateTreeModule.StateTreeNodeBlueprintBase")
 
-+EnumRedirects=(OldName="/Script/StateTreeModule.EStateTreeBindableStructSource",ValueChanges=(("TreeData", "Context"), ("TreeParameter", "Parameter"), ("StartParameter", "State"), ("NotSet", "None"))
++EnumRedirects=(OldName="/Script/StateTreeModule.EStateTreeBindableStructSource",ValueChanges=(("TreeData", "Context"), ("TreeParameter", "Parameter"), ("StartParameter", "State"), ("NotSet", "None")))
 
 +PropertyRedirects=(OldName="/Script/StateTreeModule.StateTree.NamedExternalDataDescs",NewName="/Script/StateTreeModule.StateTree.ContextDataDescs")
 
@@ -64,5 +64,5 @@
 +StructRedirects=(OldName="StateTreeEditorPropertyBinding",NewName="/script/StateTreeModule.StateTreePropertyPathBinding")
 +StructRedirects=(OldName="StateTreePropCopyBatch",NewName="/script/StateTreeModule.StateTreePropertyCopyBatch")
 
-+EnumRedirects=(OldName="/Script/StateTreeModule.EStateTreeDataSourceType",ValueChanges=(("LinkedStateParameterData", "StateParameterData"))
-+EnumRedirects=(OldName="/Script/StateTreeModule.EStateTreeBindableStructSource",ValueChanges=(("State", "StateParameter"))
++EnumRedirects=(OldName="/Script/StateTreeModule.EStateTreeDataSourceType",ValueChanges=(("LinkedStateParameterData", "StateParameterData")))
++EnumRedirects=(OldName="/Script/StateTreeModule.EStateTreeBindableStructSource",ValueChanges=(("State", "StateParameter")))
