Disable async/parallel text search in content browser.

diff --git a/Engine/Config/BaseEditor.ini b/Engine/Config/BaseEditor.ini
--- a/Engine/Config/BaseEditor.ini
+++ b/Engine/Config/BaseEditor.ini
@@ -326,6 +326,10 @@ FarLODDistanceTheshold=3000
 ; Workaround for crashes when opening large levels (UE-205211)
 r.GPUStatsEnabled=0 
 
+;Disable async asset view search
+AssetView.AllowAsync=0
+AssetView.AllowParallelism=0
+
 [GatherTextFromAssets]
 MinFreeMemory=1024
 MaxUsedMemory=16384
