Disable shader warnings as errors

These were not actually enabled at the moment as the cvar was also specified in ConsoleVariables.ini.

Comment the definition in ConsoleVariables.ini (the global default is already 0) which is more in line with how other variables are specified.

Remove the definition from the window ini file until we can resolve remaining issues with SM5 shaders.
#rnx

diff --git a/Engine/Config/ConsoleVariables.ini b/Engine/Config/ConsoleVariables.ini
--- a/Engine/Config/ConsoleVariables.ini
+++ b/Engine/Config/ConsoleVariables.ini
@@ -42,7 +42,7 @@
 r.DumpShaderDebugInfo=2
 
 ; When this is enabled, shader compilation warnings will emit as errors
-r.Shaders.WarningsAsErrors=0 ;; 0: off, 1: global shaders, 2: all
+; r.Shaders.WarningsAsErrors=0 ;; 0: off, 1: global shaders, 2: all
 
 ; When this is enabled, dumped shader paths will get collapsed (in the cases where paths are longer than the OS's max)
 ;r.DumpShaderDebugShortNames=1
diff --git a/Engine/Config/Windows/BaseWindowsEngine.ini b/Engine/Config/Windows/BaseWindowsEngine.ini
--- a/Engine/Config/Windows/BaseWindowsEngine.ini
+++ b/Engine/Config/Windows/BaseWindowsEngine.ini
@@ -35,5 +35,3 @@ MaxNumThreadsWithTlsSlots=512
 [ConsoleVariables]
 ; necessary for LowMemoryWarning events to be fired
 memory.MemoryPressureCriticalThresholdMB=512
-; Enable warnings as errors for global shaders only
-r.Shaders.WarningsAsErrors=1
