Heterogeneous Volumes: Initializing scalability defaults

diff --git a/Engine/Config/BaseScalability.ini b/Engine/Config/BaseScalability.ini
--- a/Engine/Config/BaseScalability.ini
+++ b/Engine/Config/BaseScalability.ini
@@ -643,6 +643,9 @@ r.SkyAtmosphere.TransmittanceLUT.UseSmallFormat=1
 r.SkyAtmosphere.TransmittanceLUT.SampleCount=10.0
 r.SkyAtmosphere.MultiScatteringLUT.SampleCount=15.0
 fx.Niagara.QualityLevel=0
+r.HeterogeneousVolumes.DownsampleFactor=8
+r.HeterogeneousVolumes.MaxStepCount=128
+r.HeterogeneousVolumes.Shadows.Resolution=128
 
 [EffectsQuality@1]
 r.TranslucencyLightingVolumeDim=32
@@ -670,6 +673,9 @@ r.SkyAtmosphere.TransmittanceLUT.UseSmallFormat=0
 r.SkyAtmosphere.TransmittanceLUT.SampleCount=10.0
 r.SkyAtmosphere.MultiScatteringLUT.SampleCount=15.0
 fx.Niagara.QualityLevel=1
+r.HeterogeneousVolumes.DownsampleFactor=4
+r.HeterogeneousVolumes.MaxStepCount=256
+r.HeterogeneousVolumes.Shadows.Resolution=128
 
 [EffectsQuality@2]
 r.TranslucencyLightingVolumeDim=48
@@ -697,6 +703,9 @@ r.SkyAtmosphere.TransmittanceLUT.UseSmallFormat=0
 r.SkyAtmosphere.TransmittanceLUT.SampleCount=10.0
 r.SkyAtmosphere.MultiScatteringLUT.SampleCount=15.0
 fx.Niagara.QualityLevel=2
+r.HeterogeneousVolumes.DownsampleFactor=2
+r.HeterogeneousVolumes.MaxStepCount=256
+r.HeterogeneousVolumes.Shadows.Resolution=256
 
 [EffectsQuality@3]
 r.TranslucencyLightingVolumeDim=64
@@ -725,6 +734,9 @@ r.SkyAtmosphere.TransmittanceLUT.SampleCount=10.0
 r.SkyAtmosphere.MultiScatteringLUT.SampleCount=15.0
 fx.Niagara.QualityLevel=3
 r.Refraction.OffsetQuality=1
+r.HeterogeneousVolumes.DownsampleFactor=1
+r.HeterogeneousVolumes.MaxStepCount=512
+r.HeterogeneousVolumes.Shadows.Resolution=512
 
 [EffectsQuality@Cine]
 r.TranslucencyLightingVolumeDim=64
@@ -753,6 +765,9 @@ r.SkyAtmosphere.TransmittanceLUT.SampleCount=30.0
 r.SkyAtmosphere.MultiScatteringLUT.SampleCount=20.0
 fx.Niagara.QualityLevel=4
 r.Refraction.OffsetQuality=1
+r.HeterogeneousVolumes.DownsampleFactor=1
+r.HeterogeneousVolumes.MaxStepCount=1024
+r.HeterogeneousVolumes.Shadows.Resolution=1024
 
 ;-----------------------------------------------------------------------------------------------------------------
 
