[shaders] add commented default value for warnings filter + explanation
#rnx

diff --git a/Engine/Config/ConsoleVariables.ini b/Engine/Config/ConsoleVariables.ini
--- a/Engine/Config/ConsoleVariables.ini
+++ b/Engine/Config/ConsoleVariables.ini
@@ -32,7 +32,7 @@
 [Startup]
 
 
-; Uncomment to get detailed logs on shader compiles and the opportunity to retry on errors
+; Uncomment to get detailed logs on shader compiles and the opportunity to retry on errors, and enable showing shader warnings
 ;r.ShaderDevelopmentMode=1
 
 ; Dump shaders in the Saved folder
@@ -42,7 +42,13 @@
 r.DumpShaderDebugInfo=2
 
 ; When this is enabled, shader compilation warnings will emit as errors
-; r.Shaders.WarningsAsErrors=0 ;; 0: off, 1: global shaders, 2: all
+;r.Shaders.WarningsAsErrors=0 ;; 0: off, 1: global shaders, 2: all
+
+; When this is enabled, shader compilation warnings will be logged as per the filter setting below (this cvar is ignored if shader dev mode is on)
+;r.ShowShaderCompilerWarnings=1
+
+; Filters which warnings are shown when showing shader compilation warnings; 0=none, 1=global shaders only, 2=all (default)
+;r.ShaderCompilerWarningsFilter=2
 
 ; When this is enabled, dumped shader paths will get collapsed (in the cases where paths are longer than the OS's max)
 ;r.DumpShaderDebugShortNames=1
