IterativeCook: Add NiagaraSystem and SoundWave to the iterativeclassdenylist.
#rnx

diff --git a/Engine/Config/BaseEditor.ini b/Engine/Config/BaseEditor.ini
--- a/Engine/Config/BaseEditor.ini
+++ b/Engine/Config/BaseEditor.ini
@@ -494,7 +494,9 @@ SaveUnversioned=true
 +IterativeClassDenyList=/Script/CustomizableObject.CustomizableObject		; UE-209204
 +IterativeClassDenyList=/Script/Engine.BlueprintGeneratedClass 				; Not yet documented
 +IterativeClassDenyList=/Script/Engine.BodySetup 							; UE-202602
++IterativeClassDenyList=/Script/Engine.SoundWave  							; Does an AssetRegistryQuery
 +IterativeClassDenyList=/Script/Engine.World 								; Not yet documented
++IterativeClassDenyList=/Script/Niagara.NiagaraSystem						; UE-214982
 +IterativeClassDenyList=/Script/PoseSearch.PoseSearchDatabase				; UE-209084
 +IterativeClassDenyList=/Script/RigVM.RigVM                                 ; UE-209204
 
