Enable additional network client stats.

diff --git a/Engine/Config/BaseEngine.ini b/Engine/Config/BaseEngine.ini
--- a/Engine/Config/BaseEngine.ini
+++ b/Engine/Config/BaseEngine.ini
@@ -3565,6 +3565,9 @@ IntervalSeconds=1
 +Listeners=(MetricName=OutPacketsPerFrame, Class=/Script/Engine.NetworkMetricsCSV_Replication)
 +Listeners=(MetricName=InLostPacketsFoundPerFrame, Class=/Script/Engine.NetworkMetricsCSV_Replication)
 +Listeners=(MetricName=OutLostPacketsFoundPerFrame, Class=/Script/Engine.NetworkMetricsCSV_Replication)
++Listeners=(MetricName=Ping, Class=/Script/Engine.NetworkMetricsCSV_Replication)
++Listeners=(MetricName=InRate, Class=/Script/Engine.NetworkMetricsCSV_Replication)
++Listeners=(MetricName=OutRate, Class=/Script/Engine.NetworkMetricsCSV_Replication)
 +Listeners=(MetricName=NumConnections, Class=/Script/Engine.NetworkMetricsPerfCounters)
 +Listeners=(MetricName=AvgPing, Class=/Script/Engine.NetworkMetricsPerfCounters)
 +Listeners=(MetricName=InBunches, Class=/Script/Engine.NetworkMetricsPerfCounters)
