[MemReport] Restore behavior of -full parameter by duplicating the list of commands from default profile

diff --git a/Engine/Config/BaseEngine.ini b/Engine/Config/BaseEngine.ini
--- a/Engine/Config/BaseEngine.ini
+++ b/Engine/Config/BaseEngine.ini
@@ -2188,7 +2188,43 @@ bRecompileOnLoad=true
 +Cmd="wp.DumpStreamingSources"
 
 [MemReportFullCommands]
-; Additional commands to run with memreport -full
+; These commands are run when memreport is executed with `-full` (for backward comaptibility) or `-type=Full` parameter,
+; and output to a profile file.
+; Also you can create your own sections and name them [MemReport<Profile>Commands]
+; then you will be able to execute them via `memreport -type=<Profile>`
++Cmd="Mem FromReport"
++Cmd="LogCountedInstances"
++Cmd="obj list -resourcesizesort"
++Cmd="rhi.DumpMemory"
++Cmd="rhi.DumpResourceMemory"
++Cmd="rhi.dumpresourcememory summary name=Lumen"
++Cmd="rhi.dumpresourcememory summary name=Nanite"
++Cmd="rhi.dumpresourcememory summary name=Shadow"
++Cmd="rhi.dumpresourcememory summary name=DistanceFields"
++Cmd="rhi.dumpresourcememory summary name=GPUScene"
++Cmd="rhi.dumpresourcememory summary name=SkinCache"
++Cmd="rhi.dumpresourcememory summary name=TangentsRHIBuffer"
++Cmd="rhi.dumpresourcememory summary name=IndexBuffer"
++Cmd="rhi.dumpresourcememory summary name=VertexBuffer"
++Cmd="rhi.dumpresourcememory summary name=TexCoordRHIBuffer"
++Cmd="rhi.dumpresourcememory summary name=None"
++Cmd="rhi.dumpresourcememory summary Type=TextureCube"
++Cmd="rhi.dumpresourcememory summary name=CaptureCube"
++Cmd="rhi.dumpresourcememory summary name=MorphDeltasVB"
++Cmd="rhi.dumpresourcememory summary name=VertexIndicesVB"
++Cmd="rhi.dumpresourcememory summary name=VirtualTexture"
++Cmd="rhi.dumpresourcememory summary name=VirtualPhysicalTexture"
++Cmd="rhi.dumpresourcememory summary name=VolumetricFog"
++Cmd="rhi.dumpresourcememory summary name=Hair"
++Cmd="rhi.dumpresourcememory summary name=Subsurface"
++Cmd="obj list class=SkeletalMesh -resourcesizesort"
++Cmd="obj list class=StaticMesh -resourcesizesort"
++Cmd="listtextures nonvt"
++Cmd="listtextures uncompressed"
++Cmd="LogOutStatLevels"
++Cmd="ListSpawnedActors"
++Cmd="wp.DumpDataLayers"
++Cmd="wp.DumpStreamingSources"
 +Cmd="DumpParticleMem"
 +Cmd="ConfigMem"
 +Cmd="r.DumpRenderTargetPoolMemory"
