Add the EVCF_Scalability flag to foliage.LODDistanceScale

diff --git a/Engine/Config/BaseScalability.ini b/Engine/Config/BaseScalability.ini
--- a/Engine/Config/BaseScalability.ini
+++ b/Engine/Config/BaseScalability.ini
@@ -776,30 +776,35 @@ foliage.DensityScale=0
 grass.DensityScale=0
 ;foliage.DiscardDataOnLoad=1 ; having this enabled will reduce ram use but will require reloading the level when changing FoliageQuality setting
 ;grass.DiscardDataOnLoad=1   ; having this enabled will reduce ram use but will require reloading the level when changing FoliageQuality setting
+;foliage.LODDistanceScale=0
 
 [FoliageQuality@1]
 foliage.DensityScale=0.4
 grass.DensityScale=0.4
 ;foliage.DiscardDataOnLoad=0
 ;grass.DiscardDataOnLoad=0
+;foliage.LODDistanceScale=0.4
 
 [FoliageQuality@2]
 foliage.DensityScale=0.8
 grass.DensityScale=0.8
 ;foliage.DiscardDataOnLoad=0
 ;grass.DiscardDataOnLoad=0
+;foliage.LODDistanceScale=0.8
 
 [FoliageQuality@3]
 foliage.DensityScale=1.0
 grass.DensityScale=1.0
 ;foliage.DiscardDataOnLoad=0
 ;grass.DiscardDataOnLoad=0
+;foliage.LODDistanceScale=1.0
 
 [FoliageQuality@Cine]
 foliage.DensityScale=1.0
 grass.DensityScale=1.0
 ;foliage.DiscardDataOnLoad=0
 ;grass.DiscardDataOnLoad=0
+;foliage.LODDistanceScale=1.0
 
 ;-----------------------------------------------------------------------------------------------------------------
 
diff --git a/Engine/Source/Runtime/Engine/Private/HierarchicalInstancedStaticMesh.cpp b/Engine/Source/Runtime/Engine/Private/HierarchicalInstancedStaticMesh.cpp
--- a/Engine/Source/Runtime/Engine/Private/HierarchicalInstancedStaticMesh.cpp
+++ b/Engine/Source/Runtime/Engine/Private/HierarchicalInstancedStaticMesh.cpp
@@ -98,8 +98,8 @@ TAutoConsoleVariable<float> CVarFoliageLODDistanceScale(
 	FConsoleVariableDelegate::CreateLambda([](IConsoleVariable* InVariable)
 		{
 			FGlobalComponentRecreateRenderStateContext Context;
-		})
-	);
+		}),
+	ECVF_Scalability);
 
 TAutoConsoleVariable<float> CVarRandomLODRange(
 	TEXT("foliage.RandomLODRange"),
