PR #11954: Add Immortalis-G925 device profile

diff --git a/Engine/Config/BaseDeviceProfiles.ini b/Engine/Config/BaseDeviceProfiles.ini
--- a/Engine/Config/BaseDeviceProfiles.ini
+++ b/Engine/Config/BaseDeviceProfiles.ini
@@ -118,6 +118,8 @@
 +DeviceProfileNameAndTypes=Android_Mali_G710_Vulkan,Android
 +DeviceProfileNameAndTypes=Android_Mali_G7xx,Android
 +DeviceProfileNameAndTypes=Android_Mali_G7xx_Vulkan,Android
++DeviceProfileNameAndTypes=Android_Mali_G9xx,Android
++DeviceProfileNameAndTypes=Android_Mali_G9xx_Vulkan,Android
 +DeviceProfileNameAndTypes=Android_Xclipse_5xx,Android
 +DeviceProfileNameAndTypes=Android_Xclipse_5xx_Vulkan,Android
 +DeviceProfileNameAndTypes=Android_Xclipse_9xx,Android
@@ -739,9 +741,11 @@ BaseProfileName=VisionPro
 +MatchProfile=(Profile="Android_Mali_G78",Match=((SourceType=SRC_GpuFamily,CompareType=CMP_Regex,MatchString="^Mali\\-G78($| )")))
 +MatchProfile=(Profile="Android_Mali_G710_Vulkan",Match=((SourceType=SRC_GpuFamily,CompareType=CMP_Regex,MatchString="^Mali\\-G710"),(SourceType=SRC_VulkanAvailable,CompareType=CMP_Equal,MatchString="true")))
 +MatchProfile=(Profile="Android_Mali_G710",Match=((SourceType=SRC_GpuFamily,CompareType=CMP_Regex,MatchString="^Mali\\-G710")))
-+MatchProfile=(Profile="Android_Mali_Vulkan_SM5",Match=((SourceType=SRC_GpuFamily,CompareType=CMP_Regex,MatchString="^Mali\\-G7[1-9][0-9]"),(SourceType=SRC_SM5Available,CompareType=CMP_Equal,MatchString="true")))
++MatchProfile=(Profile="Android_Mali_Vulkan_SM5",Match=((SourceType=SRC_GpuFamily,CompareType=CMP_Regex,MatchString="^Mali\\-G[7-9][1-9][0-9]"),(SourceType=SRC_SM5Available,CompareType=CMP_Equal,MatchString="true")))
 +MatchProfile=(Profile="Android_Mali_G7xx_Vulkan",Match=((SourceType=SRC_GpuFamily,CompareType=CMP_Regex,MatchString="^Mali\\-G7[0-9][0-9]"),(SourceType=SRC_VulkanAvailable,CompareType=CMP_Equal,MatchString="true")))
 +MatchProfile=(Profile="Android_Mali_G7xx",Match=((SourceType=SRC_GpuFamily,CompareType=CMP_Regex,MatchString="^Mali\\-G7[0-9][0-9]")))
++MatchProfile=(Profile="Android_Mali_G9xx_Vulkan",Match=((SourceType=SRC_GpuFamily,CompareType=CMP_Regex,MatchString="^Mali\\-G9[0-9][0-9]"),(SourceType=SRC_VulkanAvailable,CompareType=CMP_Equal,MatchString="true")))
++MatchProfile=(Profile="Android_Mali_G9xx",Match=((SourceType=SRC_GpuFamily,CompareType=CMP_Regex,MatchString="^Mali\\-G9[0-9][0-9]")))
 
 ; Enable Vulkan on any unknown Mali running Android 10 or later
 +MatchProfile=(Profile="Android_Unknown_Vulkan",Match=((SourceType=SRC_GpuFamily,CompareType=CMP_Regex,MatchString="^Mali"),(SourceType=SRC_AndroidVersion, CompareType=CMP_Regex,MatchString="([0-9]+).*"),(SourceType=SRC_PreviousRegexMatch,CompareType=CMP_GreaterEqual,MatchString="10"),(SourceType=SRC_VulkanAvailable,CompareType=CMP_Equal,MatchString="true")))
@@ -1094,6 +1098,16 @@ BaseProfileName=Android_Mali_G7xx
 ; enable Vulkan
 +CVars=r.Android.DisableVulkanSupport=0
 
+[Android_Mali_G9xx DeviceProfile]
+DeviceType=Android
+BaseProfileName=Android_Mali_G710
+
+[Android_Mali_G9xx_Vulkan DeviceProfile]
+DeviceType=Android
+BaseProfileName=Android_Mali_G9xx
+; enable Vulkan
++CVars=r.Android.DisableVulkanSupport=0
+
 [Android_TegraK1 DeviceProfile]
 DeviceType=Android
 BaseProfileName=Android_High
