[GroomBinding]
 - Enable async compilation by default

[RN] Groom Bindings are now compiled asynchronously by default in the editor

diff --git a/Engine/Config/BaseEditorPerProjectUserSettings.ini b/Engine/Config/BaseEditorPerProjectUserSettings.ini
--- a/Engine/Config/BaseEditorPerProjectUserSettings.ini
+++ b/Engine/Config/BaseEditorPerProjectUserSettings.ini
@@ -157,6 +157,7 @@ bEnableAsyncTextureCompilation=True
 bEnableAsyncStaticMeshCompilation=True
 bEnableAsyncSkinnedAssetCompilation=True
 bEnableAsyncSoundWaveCompilation=True
+bEnableAsyncGroomBindingCompilation=True
 
 ; Cooking in the editor is disabled by default now; WorldPartition does not yet support it and we want to deprecate the mode in the future to reduce the number of code flows that have to be supported
 bDisableCookInEditor=true
