Add new iPad devices (iPadAir 6th Gen, iPadPro11 5th Gen and iPadPro12.9 7th Gen)

#rnx

diff --git a/Engine/Config/BaseDeviceProfiles.ini b/Engine/Config/BaseDeviceProfiles.ini
--- a/Engine/Config/BaseDeviceProfiles.ini
+++ b/Engine/Config/BaseDeviceProfiles.ini
@@ -13,6 +13,7 @@
 +DeviceProfileNameAndTypes=iPadAir3,IOS
 +DeviceProfileNameAndTypes=iPadAir4,IOS
 +DeviceProfileNameAndTypes=iPadAir5,IOS
++DeviceProfileNameAndTypes=iPadAir6,IOS
 +DeviceProfileNameAndTypes=iPadMini4,IOS
 +DeviceProfileNameAndTypes=iPadMini5,IOS
 +DeviceProfileNameAndTypes=iPadMini6,IOS
@@ -64,10 +65,12 @@
 +DeviceProfileNameAndTypes=iPadPro2_11,IOS
 +DeviceProfileNameAndTypes=iPadPro3_11,IOS
 +DeviceProfileNameAndTypes=iPadPro4_11,IOS
++DeviceProfileNameAndTypes=iPadPro5_11,IOS
 +DeviceProfileNameAndTypes=iPadPro3_129,IOS
 +DeviceProfileNameAndTypes=iPadPro4_129,IOS
 +DeviceProfileNameAndTypes=iPadPro5_129,IOS
 +DeviceProfileNameAndTypes=iPadPro6_129,IOS
++DeviceProfileNameAndTypes=iPadPro7_129,IOS
 +DeviceProfileNameAndTypes=AppleTV,IOS
 +DeviceProfileNameAndTypes=AppleTV4K,IOS
 +DeviceProfileNameAndTypes=AppleTV2_4K,IOS
@@ -267,6 +270,9 @@ iPad13,1[8-9]=iPad10
 iPad14,[1,2]=iPadMini6
 iPad14,[3,4]=iPadPro4_11
 iPad14,[5,6]=iPadPro6_129
+iPad14,[8,9]=iPadAir6
+iPad16,[3,4]=iPadPro5_11
+iPad16,[5,6]=iPadPro7_129
 AppleTV5,=AppleTV
 AppleTV6,=AppleTV4K
 AppleTV11,=AppleTV2_4K
@@ -581,6 +587,10 @@ BaseProfileName=IPadPro
 DeviceType=IOS
 BaseProfileName=IPadPro3_11
 
+[iPadPro5_11 DeviceProfile]
+DeviceType=IOS
+BaseProfileName=IPadPro3_11
+
 [iPad5 DeviceProfile]
 DeviceType=IOS
 BaseProfileName=IPadPro97
@@ -634,6 +644,10 @@ BaseProfileName=IPadPro
 DeviceType=IOS
 BaseProfileName=IPadPro5_129
 
+[iPadPro7_129 DeviceProfile]
+DeviceType=IOS
+BaseProfileName=IPadPro5_129
+
 [iPadAir3 DeviceProfile]
 DeviceType=IOS
 BaseProfileName=IPadPro
@@ -648,6 +662,10 @@ BaseProfileName=IPadPro
 DeviceType=IOS
 BaseProfileName=IPadAir4
 
+[iPadAir6 DeviceProfile]
+DeviceType=IOS
+BaseProfileName=IPadAir4
+
 [iPadMini5 DeviceProfile]
 DeviceType=IOS
 BaseProfileName=IPadPro
