Prevent .pem files using plugin based compression when being added to a pak file as IAS can end up trying to load a .pem file before compression plugins have had a chance to load.

#rnx

diff --git a/Engine/Config/BaseEngine.ini b/Engine/Config/BaseEngine.ini
--- a/Engine/Config/BaseEngine.ini
+++ b/Engine/Config/BaseEngine.ini
@@ -3020,6 +3020,7 @@ bDoUseOodleDespiteNoPluginCompression = true
 +ExtensionsToNotUsePluginCompression=ini
 +ExtensionsToNotUsePluginCompression=icu
 +ExtensionsToNotUsePluginCompression=res
++ExtensionsToNotUsePluginCompression=pem
 +FileNamesToNotUsePluginCompression=AssetRegistry.bin
 +DirectoryIndexKeepFiles="*/Config/Tags/*"
 +DirectoryIndexKeepFiles="*/Content/Localization/*"
