CIS Issue 501802: added missing redirectors after moving StructUtils

diff --git a/Engine/Config/BaseEngine.ini b/Engine/Config/BaseEngine.ini
--- a/Engine/Config/BaseEngine.ini
+++ b/Engine/Config/BaseEngine.ini
@@ -1271,7 +1271,8 @@ DefaultWorldPartitionHLODLayer=/Engine/Maps/Templates/HLODs/HLODLayer_Instanced.
 +PackageRedirects=(OldName="/Script/WindowsTargetPlatform", NewName="/Script/WindowsTargetPlatformSettings")
 +PackageRedirects=(OldName="/Script/MacTargetPlatform", NewName="/Script/MacTargetPlatformSettings")
 +PackageRedirects=(OldName="/Script/StructUtils", NewName="/Script/CoreUObject")
-
++PackageRedirects=(OldName="/Script/StructUtilsEngine", NewName="/Script/Engine")
++PackageRedirects=(OldName="/Script/StructUtilsNodes", NewName="/Script/BlueprintGraph")
 
 +ClassRedirects=(OldName="/Script/Engine.UserDefinedStruct",NewName="/Script/CoreUObject.UserDefinedStruct")
 +ClassRedirects=(OldName="/Script/StructUtils.StructUtilsFunctionLibrary",NewName="/Script/Engine.BlueprintInstancedStructLibrary")
