Remove redundant CVars from Quest device profiles.

diff --git a/Engine/Config/BaseDeviceProfiles.ini b/Engine/Config/BaseDeviceProfiles.ini
--- a/Engine/Config/BaseDeviceProfiles.ini
+++ b/Engine/Config/BaseDeviceProfiles.ini
@@ -1141,8 +1141,6 @@ BaseProfileName=Android_High
 [Meta_Quest_3 DeviceProfile]
 DeviceType=Android
 BaseProfileName=Meta_Quest_Pro
-+CVars=fx.NiagaraAllowGPUParticles=1
-+CVars=FX.AllowGPUSorting=1
 
 [Meta_Quest_Pro DeviceProfile]
 DeviceType=Android
