Enables r.TSR.ReprojectionField on high, epic and cinematic anti-aliasing quality

diff --git a/Engine/Config/BaseScalability.ini b/Engine/Config/BaseScalability.ini
--- a/Engine/Config/BaseScalability.ini
+++ b/Engine/Config/BaseScalability.ini
@@ -57,6 +57,7 @@ r.TSR.History.ScreenPercentage=100
 r.TSR.History.UpdateQuality=0
 r.TSR.ShadingRejection.Flickering=0
 r.TSR.RejectionAntiAliasingQuality=0
+r.TSR.ReprojectionField=0
 r.TSR.Resurrection=0
 
 [AntiAliasingQuality@1]
@@ -67,6 +68,7 @@ r.TSR.History.ScreenPercentage=100
 r.TSR.History.UpdateQuality=1
 r.TSR.ShadingRejection.Flickering=0
 r.TSR.RejectionAntiAliasingQuality=1
+r.TSR.ReprojectionField=0
 r.TSR.Resurrection=1
 
 [AntiAliasingQuality@2]
@@ -77,6 +79,7 @@ r.TSR.History.ScreenPercentage=100
 r.TSR.History.UpdateQuality=2
 r.TSR.ShadingRejection.Flickering=1
 r.TSR.RejectionAntiAliasingQuality=1
+r.TSR.ReprojectionField=1
 r.TSR.Resurrection=1
 
 [AntiAliasingQuality@3]
@@ -87,6 +90,7 @@ r.TSR.History.ScreenPercentage=200
 r.TSR.History.UpdateQuality=3
 r.TSR.ShadingRejection.Flickering=1
 r.TSR.RejectionAntiAliasingQuality=2
+r.TSR.ReprojectionField=1
 r.TSR.Resurrection=1
 
 [AntiAliasingQuality@Cine]
@@ -97,6 +101,7 @@ r.TSR.History.ScreenPercentage=200
 r.TSR.History.UpdateQuality=3
 r.TSR.ShadingRejection.Flickering=1
 r.TSR.RejectionAntiAliasingQuality=2
+r.TSR.ReprojectionField=1
 r.TSR.Resurrection=1
 
 
diff --git a/Engine/Source/Runtime/Renderer/Private/PostProcess/TemporalSuperResolution.cpp b/Engine/Source/Runtime/Renderer/Private/PostProcess/TemporalSuperResolution.cpp
--- a/Engine/Source/Runtime/Renderer/Private/PostProcess/TemporalSuperResolution.cpp
+++ b/Engine/Source/Runtime/Renderer/Private/PostProcess/TemporalSuperResolution.cpp
@@ -278,7 +278,7 @@ TAutoConsoleVariable<int32> CVarTSRAsyncCompute(
 
 TAutoConsoleVariable<int32> CVarTSRReprojectionField(
 	TEXT("r.TSR.ReprojectionField"), 0,
-	TEXT("Enables TSR's reprojection field for higher reprojection vector upscale and dilate quality (Experimental, disabled by default).\n")
+	TEXT("Enables TSR's reprojection field for higher reprojection vector upscale and dilate quality (Enabled by default on high, epic and cinematic anti-aliasing quality).\n")
 	TEXT("\n")
 	TEXT("When the reprojection fields is enabled, it dilates the reprojection vector by half spatially ")
 	TEXT("anti-aliased rendering pixel from the depth buffer, instead by a full rendering pixel ")
