Change the default value of Takename to match the settings used for takerecorder

diff --git a/Engine/Config/BaseEngine.ini b/Engine/Config/BaseEngine.ini
--- a/Engine/Config/BaseEngine.ini
+++ b/Engine/Config/BaseEngine.ini
@@ -3466,7 +3466,7 @@ HMDWornMovementThreshold = 50.0
 
 [/Script/Engine.AnimationSettings]
 bStripAnimationDataOnDedicatedServer=False
-BoneTimecodeCustomAttributeNameSettings=(HourAttributeName="TCHour",MinuteAttributeName="TCMinute",SecondAttributeName="TCSecond",FrameAttributeName="TCFrame",SubframeAttributeName="TCSubframe",RateAttributeName="TCRate",TakenameAttributeName="Takename")
+BoneTimecodeCustomAttributeNameSettings=(HourAttributeName="TCHour",MinuteAttributeName="TCMinute",SecondAttributeName="TCSecond",FrameAttributeName="TCFrame",SubframeAttributeName="TCSubframe",RateAttributeName="TCRate",TakenameAttributeName="TCSlate")
 
 [Animation.DefaultObjectSettings]
 ; Fallback is used only if the default codec isn't found (e.g. if the plugin isn't enabled)
