Iris
- Changed engine config for gameplaycategoryreplicator to use deltacompression and static priority

diff --git a/Engine/Config/BaseEngine.ini b/Engine/Config/BaseEngine.ini
--- a/Engine/Config/BaseEngine.ini
+++ b/Engine/Config/BaseEngine.ini
@@ -1370,6 +1370,7 @@ DefaultSpatialFilterName=Spatial
 ; Classes for which deltacompression is enabled
 +DeltaCompressionConfigs=(ClassName=/Script/Engine.Pawn))
 +DeltaCompressionConfigs=(ClassName=/Script/Engine.PlayerState))
++DeltaCompressionConfigs=(ClassName=/Script/GameplayDebugger.GameplayDebuggerCategoryReplicator))
 
 ; CriticalClassConfigs
 ; Classes which force a disconnection when a protocol mismatch is detected
