Fix typo in SheenLegacy LUT path causing cook error.

diff --git a/Engine/Config/BaseEngine.ini b/Engine/Config/BaseEngine.ini
--- a/Engine/Config/BaseEngine.ini
+++ b/Engine/Config/BaseEngine.ini
@@ -166,7 +166,7 @@ SheenLTCTextureName=/Engine/EngineMaterials/LTC/Sheen_LTC.Sheen_LTC
 GGXReflectionEnergyTextureName=/Engine/EngineMaterials/EnergyConservation/GGX_ReflectionEnergy.GGX_ReflectionEnergy
 GGXTransmissionEnergyTextureName=/Engine/EngineMaterials/EnergyConservation/GGX_TransmissionEnergyArray.GGX_TransmissionEnergyArray
 SheenEnergyTextureName=/Engine/EngineMaterials/EnergyConservation/Sheen_Energy.Sheen_Energy
-SheenLegacyEnergyTextureName=/Engine/EngineMaterials/EnergyConservation/Sheen_LegacyEnergy.Sheen_LegacyEnergy
+SheenLegacyEnergyTextureName=/Engine/EngineMaterials/EnergyConservation/SheenLegacy_Energy.SheenLegacy_Energy
 DiffuseEnergyTextureName=/Engine/EngineMaterials/EnergyConservation/Diffuse_Energy.Diffuse_Energy
 GlintTextureName=/Engine/EngineMaterials/Substrate/GLints/AGlintTex2dArray.AGlintTex2dArray
 GlintTexture2Name=/Engine/EngineMaterials/Substrate/GLints2/AGlintTex2dArray2.AGlintTex2dArray2
