Added a few additional redirects for deprecated KismetMathLibrary functions that were missed.

diff --git a/Engine/Config/BaseEngine.ini b/Engine/Config/BaseEngine.ini
--- a/Engine/Config/BaseEngine.ini
+++ b/Engine/Config/BaseEngine.ini
@@ -1295,6 +1295,9 @@ DefaultWorldPartitionHLODLayer=/Engine/Maps/Templates/HLODs/HLODLayer_Instanced.
 +EnumRedirects=(OldName="/Script/Engine.EUserDefinedStructureStatus",NewName="/Script/CoreUObject.EUserDefinedStructureStatus")
 
 +PropertyRedirects=(OldName="K2Node_CallFunction.bIsPureFunc", NewName="bDefaultsToPureFunc")
++FunctionRedirects=(OldName="KismetMathLibrary.FMod", NewName="FMod64")
++FunctionRedirects=(OldName="KismetMathLibrary.MinimumAreaRectangle", NewName="MinAreaRectangle")
++PropertyRedirects=(OldName="KismetMathLibrary.MinAreaRectangle.InVerts", NewName="InPoints")
 
 ; End of releases redirects
 
