Removed long broken shortcut for creating an Array Get item at Index node in blueprints

[RN] Removed long broken shortcut for creating an Array Get item at Index node in blueprints

diff --git a/Engine/Config/BaseEditorPerProjectUserSettings.ini b/Engine/Config/BaseEditorPerProjectUserSettings.ini
--- a/Engine/Config/BaseEditorPerProjectUserSettings.ini
+++ b/Engine/Config/BaseEditorPerProjectUserSettings.ini
@@ -787,7 +787,6 @@ DefaultBrushRadius=128
 +Node=(Class=/Script/BlueprintGraph.K2Node_MultiGate Key=M Shift=false Ctrl=false Alt=false)
 +Node=(Class=/Engine/EditorBlueprintResources/StandardMacros.StandardMacros:ForEachLoop Key=F Shift=false Ctrl=false Alt=false)
 +Node=(Class=/Engine/EditorBlueprintResources/StandardMacros.StandardMacros:DoOnce Key=O Shift=false Ctrl=false Alt=false)
-+Node=(Class=/Script/Engine.KismetArrayLibrary:Array_Get Key=A Shift=false Ctrl=false Alt=false)
 +Node=(Class=/Script/BlueprintGraph.K2Node_Knot Key=R Shift=false Ctrl=false Alt=false)
 
 [DefaultEventNodes]
