Add engine hooks for setting and loading the TextureColorViewModeMaterial.
This is the material used for debug visualization mode of the mesh paint color textures.
The material asset for this has already been submitted.

diff --git a/Engine/Config/BaseEngine.ini b/Engine/Config/BaseEngine.ini
--- a/Engine/Config/BaseEngine.ini
+++ b/Engine/Config/BaseEngine.ini
@@ -189,6 +189,7 @@ VertexColorViewModeMaterialName_AlphaAsColor=/Engine/EngineDebugMaterials/Vertex
 VertexColorViewModeMaterialName_RedOnly=/Engine/EngineDebugMaterials/VertexColorViewMode_RedOnly.VertexColorViewMode_RedOnly
 VertexColorViewModeMaterialName_GreenOnly=/Engine/EngineDebugMaterials/VertexColorViewMode_GreenOnly.VertexColorViewMode_GreenOnly
 VertexColorViewModeMaterialName_BlueOnly=/Engine/EngineDebugMaterials/VertexColorViewMode_BlueOnly.VertexColorViewMode_BlueOnly
+TextureColorViewModeMaterialName=/Engine/EngineDebugMaterials/TextureColorViewMode.TextureColorViewMode
 PhysicalMaterialMaskMaterialName=/Engine/EngineDebugMaterials/PhysicalMaterialMaskMaterial.PhysicalMaterialMaskMaterial
 TexturePaintingMaskMaterialName=/Engine/EngineDebugMaterials/TexturePaintingMaskMaterial.TexturePaintingMaskMaterial
 BoneWeightMaterialName=/Engine/EngineDebugMaterials/BoneWeightMaterial.BoneWeightMaterial
diff --git a/Engine/Source/Runtime/Engine/Classes/Engine/Engine.h b/Engine/Source/Runtime/Engine/Classes/Engine/Engine.h
--- a/Engine/Source/Runtime/Engine/Classes/Engine/Engine.h
+++ b/Engine/Source/Runtime/Engine/Classes/Engine/Engine.h
@@ -1057,6 +1057,14 @@ public:
 	UPROPERTY(globalconfig)
 	FString VertexColorViewModeMaterialName_BlueOnly;
 
+	/** Material for visualizing mesh paint texture colors on meshes in the scene */
+	UPROPERTY()
+	TObjectPtr<class UMaterial> TextureColorViewModeMaterial;
+
+	/** Path of the material for visualizing mesh paint texture colors on meshes in the scene */
+	UPROPERTY(globalconfig)
+	FString TextureColorViewModeMaterialName;
+
 #if WITH_EDITORONLY_DATA
 	/** Material used to render bone weights on skeletal meshes */
 	UPROPERTY()
diff --git a/Engine/Source/Runtime/Engine/Private/UnrealEngine.cpp b/Engine/Source/Runtime/Engine/Private/UnrealEngine.cpp
--- a/Engine/Source/Runtime/Engine/Private/UnrealEngine.cpp
+++ b/Engine/Source/Runtime/Engine/Private/UnrealEngine.cpp
@@ -3140,6 +3140,7 @@ void UEngine::InitializeObjectReferences()
 		LoadSpecialMaterial(TEXT("VertexColorViewModeMaterialName_RedOnly"), VertexColorViewModeMaterialName_RedOnly, VertexColorViewModeMaterial_RedOnly, false);
 		LoadSpecialMaterial(TEXT("VertexColorViewModeMaterialName_GreenOnly"), VertexColorViewModeMaterialName_GreenOnly, VertexColorViewModeMaterial_GreenOnly, false);
 		LoadSpecialMaterial(TEXT("VertexColorViewModeMaterialName_BlueOnly"), VertexColorViewModeMaterialName_BlueOnly, VertexColorViewModeMaterial_BlueOnly, false);
+		LoadSpecialMaterial(TEXT("TextureColorViewModeMaterialName"), TextureColorViewModeMaterialName, TextureColorViewModeMaterial, false);
 	}
 
 	// Nanite materials
