[Reference Viewer] Filter Search Results should be on by default

diff --git a/Engine/Config/BaseEditorPerProjectUserSettings.ini b/Engine/Config/BaseEditorPerProjectUserSettings.ini
--- a/Engine/Config/BaseEditorPerProjectUserSettings.ini
+++ b/Engine/Config/BaseEditorPerProjectUserSettings.ini
@@ -979,7 +979,7 @@ bIsShowManagementReferences=False
 bIsShowSearchableNames=False
 bIsShowCodePackages=False
 bIsShowDuplicates=True
-bIsShowFilteredPackagesOnly=False
+bIsShowFilteredPackagesOnly=True
 bIsCompactMode=False
 bIsShowPath=False
 bFiltersEnabled=True
