Potential fix for NiagaraValidationRuleSet having missing rules in a cooked editor

#rnx

diff --git a/Engine/Config/BaseEditor.ini b/Engine/Config/BaseEditor.ini
--- a/Engine/Config/BaseEditor.ini
+++ b/Engine/Config/BaseEditor.ini
@@ -435,6 +435,7 @@ MPCookGeneratorSplit=AnyWorker
 +AutomaticOptionalInclusionAssetType=/Script/Niagara.NiagaraSystem
 +AutomaticOptionalInclusionAssetType=/Script/NiagaraEditor.NiagaraParameterDefinitions
 +AutomaticOptionalInclusionAssetType=/Script/Niagara.NiagaraEffectType
++AutomaticOptionalInclusionAssetType=/Script/Niagara.NiagaraValidationRuleSet
 
 [EditorDomain]
 EditorDomainEnabled=false
