Add a description to the "showdebug input" autocomplete command list

#rnx

diff --git a/Engine/Config/BaseInput.ini b/Engine/Config/BaseInput.ini
--- a/Engine/Config/BaseInput.ini
+++ b/Engine/Config/BaseInput.ini
@@ -202,7 +202,7 @@ MaxScrollbackSize=1024
 +ManualAutoCompleteList=(Command="ShowDebug WEAPON",Desc=)
 +ManualAutoCompleteList=(Command="ShowDebug ANIMATION",Desc="Toggles display state of animation debug data")
 +ManualAutoCompleteList=(Command="ShowDebug BONES",Desc="Toggles display of skeletalmesh bones")
-+ManualAutoCompleteList=(Command="ShowDebug INPUT",Desc=)
++ManualAutoCompleteList=(Command="ShowDebug INPUT",Desc="Toggles display of all FKey states on the current Player Controller")
 +ManualAutoCompleteList=(Command="ShowDebug FORCEFEEDBACK",Desc="Toggles display of current force feedback values and what is contributing to that calculation")
 +ManualAutoCompleteList=(Command="ShowDebugToggleSubCategory 3DBONES",Desc="With ShowDebug Bones: Toggles bone rendering between single lines and a more complex 3D model per bone")
 +ManualAutoCompleteList=(Command="ShowDebugToggleSubCategory SYNCGROUPS",Desc="With ShowDebug Animation: Toggles display of sync group data")
