[Backout] - CL35775248

This is a partial back-out because if content was saved with the new functions, we don't want to break.  These are just the redirects and I can reintroduce them when everything is fixed.

Original CL Desc
-----------------------------------------------------------------
Modified version of the fix found in PR #12052: [????] GetOwnedGameplayTags Fix For Reusing old tag container in blueprint event graphs.

The (Blueprint) GameplayTagInterface::GetOwnedGameplayTags member function has been redirected to BlueprintGameplayTagLibrary::GetOwnedGameplayTags.  The implementation ensures that the results you get are always a newly initialized tag container, thus ensuring you do not accidentally reuse and append old results.

diff --git a/Engine/Config/BaseEngine.ini b/Engine/Config/BaseEngine.ini
--- a/Engine/Config/BaseEngine.ini
+++ b/Engine/Config/BaseEngine.ini
@@ -1304,11 +1304,6 @@ DefaultWorldPartitionHLODLayer=/Engine/Maps/Templates/HLODs/HLODLayer_Instanced.
 
 ; End of releases redirects
 
-; Fixing the output of tags for blueprints to be reset.
-+FunctionRedirects=(OldName="/Script/GameplayTags.GameplayTagAssetInterface.GetOwnedGameplayTags", NewName="BlueprintGameplayTagLibrary.GetOwnedGameplayTags")
-+PropertyRedirects=(OldName="BlueprintGameplayTagLibrary.GetOwnedGameplayTags.TagContainer", NewName="ReturnValue")
-+PropertyRedirects=(OldName="BlueprintGameplayTagLibrary.GetOwnedGameplayTags.Self", NewName="TagContainerInterface")
-
 [CoreUObject.Metadata]
 ; Note: UnrealHeaderTool should be rerun after making changes to MetadataRedirects to catch any keys specified in class headers
 +MetadataRedirects=(OldKey="K2Protected", NewKey="BlueprintProtected")
diff --git a/Engine/Source/Runtime/GameplayTags/Classes/GameplayTagAssetInterface.h b/Engine/Source/Runtime/GameplayTags/Classes/GameplayTagAssetInterface.h
--- a/Engine/Source/Runtime/GameplayTags/Classes/GameplayTagAssetInterface.h
+++ b/Engine/Source/Runtime/GameplayTags/Classes/GameplayTagAssetInterface.h
@@ -25,6 +25,7 @@ class IGameplayTagAssetInterface
 	 * 
 	 * @param OutTags	[OUT] Set of tags on the asset
 	 */
+	 UFUNCTION(BlueprintCallable, Category = GameplayTags)
 	virtual void GetOwnedGameplayTags(FGameplayTagContainer& TagContainer) const=0;
 
 	/**
