Improved Blue Noise textures using "FastNoise: Filter-Adapted Spatio-Temporal Sampling for Real-Time Rendering". Generated using default parameters. Improves quality of various stochastic effects relying on BlueNoiseScalar and BlueNoiseVec2 sampling.

#virtualized

diff --git a/Engine/Config/BaseEngine.ini b/Engine/Config/BaseEngine.ini
--- a/Engine/Config/BaseEngine.ini
+++ b/Engine/Config/BaseEngine.ini
@@ -158,8 +158,8 @@ GeomMaterialName=/Engine/EngineDebugMaterials/GeomMaterial.GeomMaterial
 DebugMeshMaterialName=/Engine/EngineDebugMaterials/DebugMeshMaterial.DebugMeshMaterial
 EmissiveMeshMaterialName=/Engine/EngineMaterials/EmissiveMeshMaterial.EmissiveMeshMaterial
 PreIntegratedSkinBRDFTextureName=/Engine/EngineMaterials/PreintegratedSkinBRDF.PreintegratedSkinBRDF
-BlueNoiseScalarTextureName=/Engine/EngineMaterials/STBlueNoise_scalar_128x128x64.STBlueNoise_scalar_128x128x64
-BlueNoiseVec2TextureName=/Engine/EngineMaterials/STBlueNoise_vec2_128x128x64.STBlueNoise_vec2_128x128x64
+BlueNoiseScalarTextureName=/Engine/EngineMaterials/FastBlueNoise_scalar_128x128x64.FastBlueNoise_scalar_128x128x64
+BlueNoiseVec2TextureName=/Engine/EngineMaterials/FastBlueNoise_vec2_128x128x64.FastBlueNoise_vec2_128x128x64
 GGXLTCAmpTextureName=/Engine/EngineMaterials/LTC/GGX_LTCAmp.GGX_LTCAmp
 GGXLTCMatTextureName=/Engine/EngineMaterials/LTC/GGX_LTCMat.GGX_LTCMat
 SheenLTCTextureName=/Engine/EngineMaterials/LTC/Sheen_LTC.Sheen_LTC
