DDC - Added ability to override OAuthIdentifier for Cloud DDC layer with env variables or command line arguments

diff --git a/Engine/Config/BaseEngine.ini b/Engine/Config/BaseEngine.ini
--- a/Engine/Config/BaseEngine.ini
+++ b/Engine/Config/BaseEngine.ini
@@ -2633,7 +2633,7 @@ Local=(Type=FileSystem, DeleteOnly=true, UnusedFileAge=8, Path=%ENGINEDIR%Derive
 ; Configure the shared cache that is accessed after local.
 Shared=(Type=FileSystem, UnusedFileAge=10, FoldersToClean=10, ConsiderSlowAt=70, Path=?EpicDDC, EnvPathOverride=UE-SharedDataCachePath, EditorOverrideSetting=SharedDerivedDataCache, CommandLineOverride=SharedDataCachePath)
 ; Configure the cloud cache that is accessed after shared.
-Cloud=(Type=Cloud, ServerID=Cloud, EnvHostOverride=UE-CloudDataCacheHost, CommandLineHostOverride=CloudDataCacheHost, CommandLineOAuthSecretOverride="CloudDataCacheOAuthSecret", OAuthAccessTokenEnvOverride="UE-CloudDataCacheAccessToken", EnvHttpVersionOverride="UE-CloudDataCacheHttpVersion", CommandLineHttpVersionOverride="CloudDataCacheHttpVersion")
+Cloud=(Type=Cloud, ServerID=Cloud, EnvHostOverride=UE-CloudDataCacheHost, CommandLineHostOverride=CloudDataCacheHost, CommandLineOAuthSecretOverride="CloudDataCacheOAuthSecret", OAuthAccessTokenEnvOverride="UE-CloudDataCacheAccessToken", EnvHttpVersionOverride="UE-CloudDataCacheHttpVersion", CommandLineHttpVersionOverride="CloudDataCacheHttpVersion", EnvOAuthProviderIdentifierOverride="UE-CloudDataCacheOAuthProviderIdentifier", CommandLineOAuthProviderIdentifierOverride="CloudDataCacheOAuthProviderIdentifier")
 ; Configure a Project Pak node. This is a pre-generated DDC data file for the project that can be distributed to reduce runtime fetches/generation
 ; See documentation for how to create a DDP via the DerivedDataCache commandlet
 Pak=(Type=ReadPak, Filename="%GAMEDIR%DerivedDataCache/DDC.ddp")
@@ -2648,7 +2648,7 @@ AsyncPut=(Type=AsyncPut, Inner=Hierarchy)
 Hierarchy=(Type=Hierarchical, Inner=Pak, Inner=CompressedPak, Inner=EnginePak, Inner=EnterprisePak, Inner=ZenLocal, Inner=Local, Inner=ZenShared, Inner=Shared, Inner=Cloud)
 Local=(Type=FileSystem, DeleteOnly=true, UnusedFileAge=8, Path="%ENGINEVERSIONAGNOSTICUSERDIR%DerivedDataCache", EnvPathOverride=UE-LocalDataCachePath, EditorOverrideSetting=LocalDerivedDataCache)
 Shared=(Type=FileSystem, UnusedFileAge=10, FoldersToClean=10, ConsiderSlowAt=70, Path=?EpicDDC, EnvPathOverride=UE-SharedDataCachePath, EditorOverrideSetting=SharedDerivedDataCache, CommandLineOverride=SharedDataCachePath)
-Cloud=(Type=Cloud, ServerID=Cloud, EnvHostOverride=UE-CloudDataCacheHost, CommandLineHostOverride=CloudDataCacheHost, CommandLineOAuthSecretOverride="CloudDataCacheOAuthSecret", OAuthAccessTokenEnvOverride="UE-CloudDataCacheAccessToken", EnvHttpVersionOverride="UE-CloudDataCacheHttpVersion", CommandLineHttpVersionOverride="CloudDataCacheHttpVersion")
+Cloud=(Type=Cloud, ServerID=Cloud, EnvHostOverride=UE-CloudDataCacheHost, CommandLineHostOverride=CloudDataCacheHost, CommandLineOAuthSecretOverride="CloudDataCacheOAuthSecret", OAuthAccessTokenEnvOverride="UE-CloudDataCacheAccessToken", EnvHttpVersionOverride="UE-CloudDataCacheHttpVersion", CommandLineHttpVersionOverride="CloudDataCacheHttpVersion", EnvOAuthProviderIdentifierOverride="UE-CloudDataCacheOAuthProviderIdentifier", CommandLineOAuthProviderIdentifierOverride="CloudDataCacheOAuthProviderIdentifier")
 Pak=(Type=ReadPak, Filename="%GAMEDIR%DerivedDataCache/DDC.ddp")
 CompressedPak=(Type=ReadPak, Filename="%GAMEDIR%DerivedDataCache/Compressed.ddp", Compressed=true)
 EnginePak=(Type=ReadPak, Filename=../../../Engine/DerivedDataCache/Compressed.ddp, Compressed=true)
@@ -2660,7 +2660,7 @@ ZenShared=(Type=Zen, ServerID=Shared)
 Root=(Type=Hierarchical, Inner=Pak, Inner=EnginePak, Inner=Local, Inner=ZenShared, Inner=Shared, Inner=Cloud)
 Local=(Type=FileSystem, UnusedFileAge=34, PromptIfMissing=true, Path=%ENGINEDIR%DerivedDataCache, EnvPathOverride=UE-LocalDataCachePath, EditorOverrideSetting=LocalDerivedDataCache, CommandLineOverride=LocalDataCachePath)
 Shared=(Type=FileSystem, UnusedFileAge=10, FoldersToClean=10, ConsiderSlowAt=70, Path=?EpicDDC, EnvPathOverride=UE-SharedDataCachePath, EditorOverrideSetting=SharedDerivedDataCache, CommandLineOverride=SharedDataCachePath)
-Cloud=(Type=Cloud, ServerID=Cloud, EnvHostOverride=UE-CloudDataCacheHost, CommandLineHostOverride=CloudDataCacheHost, CommandLineOAuthSecretOverride="CloudDataCacheOAuthSecret", OAuthAccessTokenEnvOverride="UE-CloudDataCacheAccessToken", EnvHttpVersionOverride="UE-CloudDataCacheHttpVersion", CommandLineHttpVersionOverride="CloudDataCacheHttpVersion")
+Cloud=(Type=Cloud, ServerID=Cloud, EnvHostOverride=UE-CloudDataCacheHost, CommandLineHostOverride=CloudDataCacheHost, CommandLineOAuthSecretOverride="CloudDataCacheOAuthSecret", OAuthAccessTokenEnvOverride="UE-CloudDataCacheAccessToken", EnvHttpVersionOverride="UE-CloudDataCacheHttpVersion", CommandLineHttpVersionOverride="CloudDataCacheHttpVersion", EnvOAuthProviderIdentifierOverride="UE-CloudDataCacheOAuthProviderIdentifier", CommandLineOAuthProviderIdentifierOverride="CloudDataCacheOAuthProviderIdentifier")
 Pak=(Type=ReadPak, Filename="%GAMEDIR%DerivedDataCache/DDC.ddp")
 EnginePak=(Type=ReadPak, Filename="%ENGINEDIR%DerivedDataCache/DDC.ddp")
 ZenShared=(Type=Zen, ServerID=Shared)
@@ -2669,7 +2669,7 @@ ZenShared=(Type=Zen, ServerID=Shared)
 Root=(Type=Hierarchical, Inner=Pak, Inner=CompressedPak, Inner=EnginePak, Inner=EnterprisePak, Inner=Local, Inner=ZenShared, Inner=Shared, Inner=Cloud)
 Local=(Type=FileSystem, UnusedFileAge=34, Path="%ENGINEVERSIONAGNOSTICUSERDIR%DerivedDataCache", EnvPathOverride=UE-LocalDataCachePath, EditorOverrideSetting=LocalDerivedDataCache)
 Shared=(Type=FileSystem, UnusedFileAge=10, FoldersToClean=10, ConsiderSlowAt=70, Path=?EpicDDC, EnvPathOverride=UE-SharedDataCachePath, EditorOverrideSetting=SharedDerivedDataCache, CommandLineOverride=SharedDataCachePath)
-Cloud=(Type=Cloud, ServerID=Cloud, EnvHostOverride=UE-CloudDataCacheHost, CommandLineHostOverride=CloudDataCacheHost, CommandLineOAuthSecretOverride="CloudDataCacheOAuthSecret", OAuthAccessTokenEnvOverride="UE-CloudDataCacheAccessToken", EnvHttpVersionOverride="UE-CloudDataCacheHttpVersion", CommandLineHttpVersionOverride="CloudDataCacheHttpVersion")
+Cloud=(Type=Cloud, ServerID=Cloud, EnvHostOverride=UE-CloudDataCacheHost, CommandLineHostOverride=CloudDataCacheHost, CommandLineOAuthSecretOverride="CloudDataCacheOAuthSecret", OAuthAccessTokenEnvOverride="UE-CloudDataCacheAccessToken", EnvHttpVersionOverride="UE-CloudDataCacheHttpVersion", CommandLineHttpVersionOverride="CloudDataCacheHttpVersion", EnvOAuthProviderIdentifierOverride="UE-CloudDataCacheOAuthProviderIdentifier", CommandLineOAuthProviderIdentifierOverride="CloudDataCacheOAuthProviderIdentifier")
 Pak=(Type=ReadPak, Filename="%GAMEDIR%DerivedDataCache/DDC.ddp")
 CompressedPak=(Type=ReadPak, Filename="%GAMEDIR%DerivedDataCache/Compressed.ddp", Compressed=true)
 EnginePak=(Type=ReadPak, Filename=../../../Engine/DerivedDataCache/Compressed.ddp, Compressed=true)
@@ -2687,7 +2687,7 @@ Local=(Type=FileSystem, DeleteOnly=true, UnusedFileAge=8, Path=%ENGINEDIR%Derive
 ZenLocal=(Type=Zen, ServerID=Local)
 ZenShared=(Type=Zen, ServerID=Shared)
 Shared=(Type=FileSystem, UnusedFileAge=10, FoldersToClean=10, ConsiderSlowAt=70, Path=?EpicDDC, EnvPathOverride=UE-SharedDataCachePath, EditorOverrideSetting=SharedDerivedDataCache, CommandLineOverride=SharedDataCachePath)
-Cloud=(Type=Cloud, ServerID=Cloud, EnvHostOverride=UE-CloudDataCacheHost, CommandLineHostOverride=CloudDataCacheHost, CommandLineOAuthSecretOverride="CloudDataCacheOAuthSecret", OAuthAccessTokenEnvOverride="UE-CloudDataCacheAccessToken", EnvHttpVersionOverride="UE-CloudDataCacheHttpVersion", CommandLineHttpVersionOverride="CloudDataCacheHttpVersion")
+Cloud=(Type=Cloud, ServerID=Cloud, EnvHostOverride=UE-CloudDataCacheHost, CommandLineHostOverride=CloudDataCacheHost, CommandLineOAuthSecretOverride="CloudDataCacheOAuthSecret", OAuthAccessTokenEnvOverride="UE-CloudDataCacheAccessToken", EnvHttpVersionOverride="UE-CloudDataCacheHttpVersion", CommandLineHttpVersionOverride="CloudDataCacheHttpVersion", EnvOAuthProviderIdentifierOverride="UE-CloudDataCacheOAuthProviderIdentifier", CommandLineOAuthProviderIdentifierOverride="CloudDataCacheOAuthProviderIdentifier")
 Pak=(Type=ReadPak, Filename="%GAMEDIR%DerivedDataCache/DDC.ddp")
 EnginePak=(Type=ReadPak, Filename="%ENGINEDIR%DerivedDataCache/DDC.ddp")
 
@@ -2741,7 +2741,7 @@ AsyncPut=(Type=AsyncPut, Inner=Hierarchy)
 Hierarchy=(Type=Hierarchical, Inner=PakWrite, Inner=PakRead, Inner=ZenLocal, Inner=Local, Inner=ZenShared, Inner=Shared, Inner=Cloud)
 Local=(Type=FileSystem, DeleteOnly=true, UnusedFileAge=8, Path=%ENGINEDIR%DerivedDataCache, EnvPathOverride=UE-LocalDataCachePath, EditorOverrideSetting=LocalDerivedDataCache, CommandLineOverride=LocalDataCachePath)
 Shared=(Type=FileSystem, DeleteUnused=false, ConsiderSlowAt=70, Path=?EpicDDC, EnvPathOverride=UE-SharedDataCachePath, EditorOverrideSetting=SharedDerivedDataCache, CommandLineOverride=SharedDataCachePath)
-Cloud=(Type=Cloud, ServerID=Cloud, EnvHostOverride=UE-CloudDataCacheHost, CommandLineHostOverride=CloudDataCacheHost, CommandLineOAuthSecretOverride="CloudDataCacheOAuthSecret", OAuthAccessTokenEnvOverride="UE-CloudDataCacheAccessToken", EnvHttpVersionOverride="UE-CloudDataCacheHttpVersion", CommandLineHttpVersionOverride="CloudDataCacheHttpVersion")
+Cloud=(Type=Cloud, ServerID=Cloud, EnvHostOverride=UE-CloudDataCacheHost, CommandLineHostOverride=CloudDataCacheHost, CommandLineOAuthSecretOverride="CloudDataCacheOAuthSecret", OAuthAccessTokenEnvOverride="UE-CloudDataCacheAccessToken", EnvHttpVersionOverride="UE-CloudDataCacheHttpVersion", CommandLineHttpVersionOverride="CloudDataCacheHttpVersion", EnvOAuthProviderIdentifierOverride="UE-CloudDataCacheOAuthProviderIdentifier", CommandLineOAuthProviderIdentifierOverride="CloudDataCacheOAuthProviderIdentifier")
 PakRead=(Type=ReadPak, Filename="%GAMEDIR%DerivedDataCache/DDC.ddp")
 PakWrite=(Type=WritePak, Filename="%GAMEDIR%DerivedDataCache/DDC.ddp")
 ZenLocal=(Type=Zen, ServerID=Local)
@@ -2753,7 +2753,7 @@ AsyncPut=(Type=AsyncPut, Inner=Hierarchy)
 Hierarchy=(Type=Hierarchical, Inner=EnginePak, Inner=PakWrite, Inner=PakRead, Inner=ZenLocal, Inner=Local, Inner=ZenShared, Inner=Shared, Inner=Cloud)
 Local=(Type=FileSystem, DeleteOnly=true, UnusedFileAge=8, Path=%ENGINEDIR%DerivedDataCache, EnvPathOverride=UE-LocalDataCachePath, EditorOverrideSetting=LocalDerivedDataCache, CommandLineOverride=LocalDataCachePath)
 Shared=(Type=FileSystem, DeleteUnused=false, ConsiderSlowAt=70, Path=?EpicDDC, EnvPathOverride=UE-SharedDataCachePath, EditorOverrideSetting=SharedDerivedDataCache, CommandLineOverride=SharedDataCachePath)
-Cloud=(Type=Cloud, ServerID=Cloud, EnvHostOverride=UE-CloudDataCacheHost, CommandLineHostOverride=CloudDataCacheHost, CommandLineOAuthSecretOverride="CloudDataCacheOAuthSecret", OAuthAccessTokenEnvOverride="UE-CloudDataCacheAccessToken", EnvHttpVersionOverride="UE-CloudDataCacheHttpVersion", CommandLineHttpVersionOverride="CloudDataCacheHttpVersion")
+Cloud=(Type=Cloud, ServerID=Cloud, EnvHostOverride=UE-CloudDataCacheHost, CommandLineHostOverride=CloudDataCacheHost, CommandLineOAuthSecretOverride="CloudDataCacheOAuthSecret", OAuthAccessTokenEnvOverride="UE-CloudDataCacheAccessToken", EnvHttpVersionOverride="UE-CloudDataCacheHttpVersion", CommandLineHttpVersionOverride="CloudDataCacheHttpVersion", EnvOAuthProviderIdentifierOverride="UE-CloudDataCacheOAuthProviderIdentifier", CommandLineOAuthProviderIdentifierOverride="CloudDataCacheOAuthProviderIdentifier")
 EnginePak=(Type=ReadPak, Filename=../../../Engine/DerivedDataCache/Compressed.ddp, Compressed=true)
 PakRead=(Type=ReadPak, Filename="%GAMEDIR%DerivedDataCache/Compressed.ddp", Compressed=true)
 PakWrite=(Type=WritePak, Filename="%GAMEDIR%DerivedDataCache/Compressed.ddp", Compressed=true)
@@ -2766,7 +2766,7 @@ AsyncPut=(Type=AsyncPut, Inner=Hierarchy)
 Hierarchy=(Type=Hierarchical, Inner=PakWrite, Inner=ZenLocal, Inner=Local, Inner=ZenShared, Inner=Shared, Inner=Cloud)
 Local=(Type=FileSystem, DeleteOnly=true, UnusedFileAge=8, Path=%ENGINEDIR%DerivedDataCache, EnvPathOverride=UE-LocalDataCachePath, EditorOverrideSetting=LocalDerivedDataCache, CommandLineOverride=LocalDataCachePath)
 Shared=(Type=FileSystem, DeleteUnused=false, ConsiderSlowAt=70, Path=?EpicDDC, EnvPathOverride=UE-SharedDataCachePath, EditorOverrideSetting=SharedDerivedDataCache, CommandLineOverride=SharedDataCachePath)
-Cloud=(Type=Cloud, ServerID=Cloud, EnvHostOverride=UE-CloudDataCacheHost, CommandLineHostOverride=CloudDataCacheHost, CommandLineOAuthSecretOverride="CloudDataCacheOAuthSecret", OAuthAccessTokenEnvOverride="UE-CloudDataCacheAccessToken", EnvHttpVersionOverride="UE-CloudDataCacheHttpVersion", CommandLineHttpVersionOverride="CloudDataCacheHttpVersion")
+Cloud=(Type=Cloud, ServerID=Cloud, EnvHostOverride=UE-CloudDataCacheHost, CommandLineHostOverride=CloudDataCacheHost, CommandLineOAuthSecretOverride="CloudDataCacheOAuthSecret", OAuthAccessTokenEnvOverride="UE-CloudDataCacheAccessToken", EnvHttpVersionOverride="UE-CloudDataCacheHttpVersion", CommandLineHttpVersionOverride="CloudDataCacheHttpVersion", EnvOAuthProviderIdentifierOverride="UE-CloudDataCacheOAuthProviderIdentifier", CommandLineOAuthProviderIdentifierOverride="CloudDataCacheOAuthProviderIdentifier")
 PakWrite=(Type=WritePak, Filename="%ENGINEDIR%DerivedDataCache/Compressed.ddp", Compressed=true)
 ZenLocal=(Type=Zen, ServerID=Local)
 ZenShared=(Type=Zen, ServerID=Shared)
@@ -2777,7 +2777,7 @@ AsyncPut=(Type=AsyncPut, Inner=Hierarchy)
 Hierarchy=(Type=Hierarchical, Inner=PakWrite, Inner=ZenLocal, Inner=Local, Inner=ZenShared, Inner=Shared, Inner=Cloud)
 Local=(Type=FileSystem, DeleteOnly=true, UnusedFileAge=8, Path=%ENGINEDIR%DerivedDataCache, EnvPathOverride=UE-LocalDataCachePath, EditorOverrideSetting=LocalDerivedDataCache, CommandLineOverride=LocalDataCachePath)
 Shared=(Type=FileSystem, DeleteUnused=false, ConsiderSlowAt=70, Path=?EpicDDC, EnvPathOverride=UE-SharedDataCachePath, EditorOverrideSetting=SharedDerivedDataCache, CommandLineOverride=SharedDataCachePath)
-Cloud=(Type=Cloud, ServerID=Cloud, EnvHostOverride=UE-CloudDataCacheHost, CommandLineHostOverride=CloudDataCacheHost, CommandLineOAuthSecretOverride="CloudDataCacheOAuthSecret", OAuthAccessTokenEnvOverride="UE-CloudDataCacheAccessToken", EnvHttpVersionOverride="UE-CloudDataCacheHttpVersion", CommandLineHttpVersionOverride="CloudDataCacheHttpVersion")
+Cloud=(Type=Cloud, ServerID=Cloud, EnvHostOverride=UE-CloudDataCacheHost, CommandLineHostOverride=CloudDataCacheHost, CommandLineOAuthSecretOverride="CloudDataCacheOAuthSecret", OAuthAccessTokenEnvOverride="UE-CloudDataCacheAccessToken", EnvHttpVersionOverride="UE-CloudDataCacheHttpVersion", CommandLineHttpVersionOverride="CloudDataCacheHttpVersion", EnvOAuthProviderIdentifierOverride="UE-CloudDataCacheOAuthProviderIdentifier", CommandLineOAuthProviderIdentifierOverride="CloudDataCacheOAuthProviderIdentifier")
 PakWrite=(Type=WritePak, Filename=../../../Enterprise/DerivedDataCache/Compressed.ddp, Compressed=true)
 ZenLocal=(Type=Zen, ServerID=Local)
 ZenShared=(Type=Zen, ServerID=Shared)
diff --git a/Engine/Source/Developer/DerivedDataCache/Private/HttpCacheStore.cpp b/Engine/Source/Developer/DerivedDataCache/Private/HttpCacheStore.cpp
--- a/Engine/Source/Developer/DerivedDataCache/Private/HttpCacheStore.cpp
+++ b/Engine/Source/Developer/DerivedDataCache/Private/HttpCacheStore.cpp
@@ -3234,8 +3234,24 @@ void FHttpCacheStoreParams::Parse(const TCHAR* NodeName, const TCHAR* Config)
 
 	FParse::Value(Config, TEXT("OAuthScope="), OAuthScope);
 
+    // OAuth Provider Identifier
 	FParse::Value(Config, TEXT("OAuthProviderIdentifier="), OAuthProviderIdentifier);
-
+	if (FParse::Value(Config, TEXT("EnvOAuthProviderIdentifierOverride="), OverrideName))
+	{
+		FString ProviderEnv = FPlatformMisc::GetEnvironmentVariable(*OverrideName);
+		if (!ProviderEnv.IsEmpty())
+		{
+			OAuthProviderIdentifier = ProviderEnv;
+			UE_LOG(LogDerivedDataCache, Log, TEXT("%s: Found environment override for OAuthProviderIdentifier %s=%s"), NodeName, *OverrideName, *OAuthProviderIdentifier);
+		}
+	}
+	if (FParse::Value(Config, TEXT("CommandLineOAuthProviderIdentifierOverride="), OverrideName))
+	{
+		if (FParse::Value(FCommandLine::Get(), *(OverrideName + TEXT("=")), OAuthProviderIdentifier))
+		{
+			UE_LOG(LogDerivedDataCache, Log, TEXT("%s: Found command line override for OAuthProviderIdentifier %s=%s"), NodeName, *OverrideName, *OAuthProviderIdentifier);
+		}
+	}
 	FParse::Value(Config, TEXT("OAuthAccess="), OAuthAccessToken);
 	if (FParse::Value(Config, TEXT("OAuthAccessTokenEnvOverride="), OverrideName))
 	{
