updates to IOS device profiles to add iPads

diff --git a/Engine/Config/BaseDeviceProfiles.ini b/Engine/Config/BaseDeviceProfiles.ini
--- a/Engine/Config/BaseDeviceProfiles.ini
+++ b/Engine/Config/BaseDeviceProfiles.ini
@@ -272,6 +272,7 @@ iPad14,[1,2]=iPadMini6
 iPad14,[3,4]=iPadPro4_11
 iPad14,[5,6]=iPadPro6_129
 iPad14,[8,9]=iPadAir6
+iPad14,1$=iPadAir_13
 iPad16,[3,4]=iPadPro5_11
 iPad16,[5,6]=iPadPro7_129
 AppleTV5,=AppleTV
