Substrate - removed redictor asset using PackageRedirects.

#tests none
#rnx

diff --git a/Engine/Config/BaseEngine.ini b/Engine/Config/BaseEngine.ini
--- a/Engine/Config/BaseEngine.ini
+++ b/Engine/Config/BaseEngine.ini
@@ -1228,6 +1228,10 @@ DefaultWorldPartitionHLODLayer=/Engine/Maps/Templates/HLODs/HLODLayer_Instanced.
 +ClassRedirects=(OldName="/Script/Engine.MaterialExpressionStrataThinFilm",						NewName="/Script/Engine.MaterialExpressionSubstrateThinFilm")
 +StructRedirects=(OldName="/Script/Engine.StrataMaterialInput",									NewName="/Script/Engine.SubstrateMaterialInput")
 
+; Substrate node renaming + removal of redirection uassets.
++PackageRedirects=(OldName="/Engine/Functions/Strata/SMF_UE4Disney.SMF_UE4Disney",		NewName="/Engine/Functions/Substrate/SMF_UE4Legacy.SMF_UE4Legacy")
++PackageRedirects=(OldName="/Engine/Functions/Substrate/SMF_UE4Disney.SMF_UE4Disney",	NewName="/Engine/Functions/Substrate/SMF_UE4Legacy.SMF_UE4Legacy")
+
 ; Data Validation refactor 
 +EnumRedirects=(OldName="/Script/DataValidation.EDataValidationUsecase",NewName="/Script/CoreUObject.EDataValidationUsecase")
 +FunctionRedirects=(OldName="EditorValidatorBase.ValidateLoadedAsset", NewName="EditorValidatorBase.K2_ValidateLoadedAsset")
