IncrementalCook: Reallow the iterative skip of PoseSearchDatabase, CustomizableObject, and RigVM. PoseSearchDatabase we believe is fixed, for the other two we think there are remaining issues, but we want to wait for those issues to reoccur so that we can create specific JIRAs for them.
#rnx

diff --git a/Engine/Config/BaseEditor.ini b/Engine/Config/BaseEditor.ini
--- a/Engine/Config/BaseEditor.ini
+++ b/Engine/Config/BaseEditor.ini
@@ -489,12 +489,9 @@ SaveUnversioned=true
 +IterativeClassScriptPackageAllowList=Allow,<EngineRoot>
 
 +IterativeClassDenyList=/Script/ComputeFramework.ComputeGraph				; UE-203846
-+IterativeClassDenyList=/Script/CustomizableObject.CustomizableObject		; UE-209204
 +IterativeClassDenyList=/Script/Engine.BlueprintGeneratedClass 				; Not yet documented
 +IterativeClassDenyList=/Script/Engine.BodySetup 							; UE-202602
 +IterativeClassDenyList=/Script/Engine.World 								; Not yet documented
-+IterativeClassDenyList=/Script/PoseSearch.PoseSearchDatabase				; UE-209084
-+IterativeClassDenyList=/Script/RigVM.RigVM                                 ; UE-209204
 
 ; IterativeValidatePackageIgnoreList - package names that will be displayed but not treated as an error by IterativeValidate
 ; +IterativeValidatePackageIgnoreList=/Game/PackageName
