Prevent AmbientSound (Actors) from loading in a dedicated server environment

#rnx

diff --git a/Engine/Config/BaseEngine.ini b/Engine/Config/BaseEngine.ini
--- a/Engine/Config/BaseEngine.ini
+++ b/Engine/Config/BaseEngine.ini
@@ -3280,6 +3280,7 @@ DefaultASTCQualityBySizeHQ=4
 +ClassesExcludedOnDedicatedServer=SubUVAnimation
 +ClassesExcludedOnDedicatedServer=SoundNode
 +ClassesExcludedOnDedicatedServer=GameplayEffectUIData
++ClassesExcludedOnDedicatedServer=AmbientSound
 +ClassesExcludedOnDedicatedClient=WidgetBlueprint
 +ClassesExcludedOnDedicatedClient=GroupActor
 +ClassesExcludedOnDedicatedClient=MetaData
