Restore incremental cook exclusion for CustomizableObject data type.

diff --git a/Engine/Config/BaseEditor.ini b/Engine/Config/BaseEditor.ini
--- a/Engine/Config/BaseEditor.ini
+++ b/Engine/Config/BaseEditor.ini
@@ -489,6 +489,7 @@ SaveUnversioned=true
 +IterativeClassScriptPackageAllowList=Allow,<EngineRoot>
 
 +IterativeClassDenyList=/Script/ComputeFramework.ComputeGraph				; UE-203846
++IterativeClassDenyList=/Script/CustomizableObject.CustomizableObject		; UE-224744
 +IterativeClassDenyList=/Script/Engine.BlueprintGeneratedClass 				; Not yet documented
 +IterativeClassDenyList=/Script/Engine.BodySetup 							; UE-202602
 +IterativeClassDenyList=/Script/Engine.World 								; Not yet documented
