Restore incremental cook exclusion for PoseSearchDatabase data type.

diff --git a/Engine/Config/BaseEditor.ini b/Engine/Config/BaseEditor.ini
--- a/Engine/Config/BaseEditor.ini
+++ b/Engine/Config/BaseEditor.ini
@@ -493,6 +493,7 @@ SaveUnversioned=true
 +IterativeClassDenyList=/Script/Engine.BlueprintGeneratedClass 				; Not yet documented
 +IterativeClassDenyList=/Script/Engine.BodySetup 							; UE-202602
 +IterativeClassDenyList=/Script/Engine.World 								; Not yet documented
++IterativeClassDenyList=/Script/PoseSearch.PoseSearchDatabase				; UE-224771
 
 ; IterativeValidatePackageIgnoreList - package names that will be displayed but not treated as an error by IterativeValidate
 ; +IterativeValidatePackageIgnoreList=/Game/PackageName
