Add all the new iPhone16 based devices to the base device profile

diff --git a/Engine/Config/BaseDeviceProfiles.ini b/Engine/Config/BaseDeviceProfiles.ini
--- a/Engine/Config/BaseDeviceProfiles.ini
+++ b/Engine/Config/BaseDeviceProfiles.ini
@@ -52,6 +52,10 @@
 +DeviceProfileNameAndTypes=iPhone15Plus,IOS
 +DeviceProfileNameAndTypes=iPhone15Pro,IOS
 +DeviceProfileNameAndTypes=iPhone15ProMax,IOS
++DeviceProfileNameAndTypes=iPhone16,IOS
++DeviceProfileNameAndTypes=iPhone16Plus,IOS
++DeviceProfileNameAndTypes=iPhone16Pro,IOS
++DeviceProfileNameAndTypes=iPhone16ProMax,IOS
 +DeviceProfileNameAndTypes=iPadPro,IOS
 +DeviceProfileNameAndTypes=iPadPro105,IOS
 +DeviceProfileNameAndTypes=iPadPro129,IOS
@@ -238,6 +242,10 @@ iPhone15,4=iPhone15
 iPhone15,5=iPhone15Plus
 iPhone16,1=iPhone15Pro
 iPhone16,2=iPhone15ProMax
+iPhone17,1=iPhone16Pro
+iPhone17,2=iPhone16ProMax
+iPhone17,3=iPhone16
+iPhone17,4=iPhone16Plus
 iPod9,1=iPodTouch7
 iPad5,[1,2]=iPadMini4
 iPad5,[3,4]=iPadAir2
@@ -516,6 +524,23 @@ DeviceType=IOS
 BaseProfileName=iPhone13ProMax
 +CVars=ios.PhysicalScreenDensity=460
 
+[iPhone16 DeviceProfile]
+DeviceType=IOS
+BaseProfileName=iPhone15
+
+[iPhone16Plus DeviceProfile]
+DeviceType=IOS
+BaseProfileName=iPhone15Plus
++CVars=ios.PhysicalScreenDensity=460
+
+[iPhone16Pro DeviceProfile]
+DeviceType=IOS
+BaseProfileName=iPhone15Pro
+
+[iPhone16ProMax DeviceProfile]
+DeviceType=IOS
+BaseProfileName=iPhone15ProMax
+
 [iPodTouch7 DeviceProfile]
 DeviceType=IOS
 BaseProfileName=iPhone7
