Slate SDF fonts - Device profile defaults
on behalf of Viktor Chlumsky for most of the changes.
#rnx
#tests Tested in editor

diff --git a/Engine/Config/BaseDeviceProfiles.ini b/Engine/Config/BaseDeviceProfiles.ini
--- a/Engine/Config/BaseDeviceProfiles.ini
+++ b/Engine/Config/BaseDeviceProfiles.ini
@@ -180,11 +180,15 @@ TextureLODGroups=(Group=TEXTUREGROUP_World,MinLODSize=1,MaxLODSize=16384,LODBias
 +TextureLODGroups=(Group=TEXTUREGROUP_Pixels2D,MinLODSize=1,MaxLODSize=16384,LODBias=0,MinMagFilter=point,MipFilter=point,MipGenSettings=TMGS_SimpleAverage)
 +TextureLODGroups=(Group=TEXTUREGROUP_8BitData,LODBias=0,MinMagFilter=point,MipFilter=point,MipGenSettings=TMGS_NoMipmaps)
 +TextureLODGroups=(Group=TEXTUREGROUP_16BitData,LODBias=0,MinMagFilter=point,MipFilter=point,MipGenSettings=TMGS_NoMipmaps)
++CVars=UI.SlateSDFText.RasterizationMode=Msdf
++CVars=UI.SlateSDFText.ResolutionLevel=2
 
 [Windows DeviceProfile]
 DeviceType=Windows
 BaseProfileName=
 bIsVisibleForAssets=True
++CVars=UI.SlateSDFText.RasterizationMode=Msdf
++CVars=UI.SlateSDFText.ResolutionLevel=2
 
 [WindowsEditor DeviceProfile]
 DeviceType=Windows
@@ -326,6 +330,9 @@ bIsVisibleForAssets=True
 +CVars=sg.LandscapeQuality=2
 ; LightGrid Settings
 +CVars=r.Forward.LightLinkedListCulling=0
+; Optimize Slate signed distance field text for performance
++CVars=UI.SlateSDFText.RasterizationMode=SdfApproximation
++CVars=UI.SlateSDFText.ResolutionLevel=1
 
 ; PreviewAllowlistCVars and PreviewDenyListCVars are arrays of cvars (or prefixes) that are included or excluded from being applied in editor preview
 ; If any PreviewAllowlistCVars is set, cvars are denied by default.
@@ -713,6 +720,8 @@ BaseProfileName=IOS
 +CVars=sg.GlobalIlluminationQuality=3
 +CVars=r.Mobile.AmbientOcclusionQuality=0
 +CVars=r.Mobile.PixelProjectedReflectionQuality=0
++CVars=UI.SlateSDFText.RasterizationMode=SdfApproximation
++CVars=UI.SlateSDFText.ResolutionLevel=1
 
 [TVOS DeviceProfile]
 DeviceType=TVOS
@@ -839,6 +848,9 @@ bIsVisibleForAssets=True
 +CVars=r.Forward.LightLinkedListCulling=0
 ; Scale the LODDistanceFactor when doing a sky light capture
 +CVars=r.SkylightCapture.LODDistanceScale=2
+; Optimize Slate signed distance field text for performance
++CVars=UI.SlateSDFText.RasterizationMode=SdfApproximation
++CVars=UI.SlateSDFText.ResolutionLevel=1
 ; PreviewAllowlistCVars and PreviewDenyListCVars are arrays of cvars (or prefixes) that are included or excluded from being applied in editor preview.
 ; If any PreviewAllowlistCVars is set, cvars are denied by default.
 PreviewAllowlistCVars=sg.
@@ -883,6 +895,7 @@ BaseProfileName=Android
 +CVars=sg.EffectsQuality=1
 +CVars=sg.FoliageQuality=1
 +CVars=sg.LandscapeQuality=1
++CVars=UI.SlateSDFText.ResolutionLevel=2
 
 [Android_High DeviceProfile]
 DeviceType=Android
@@ -898,6 +911,8 @@ BaseProfileName=Android
 +CVars=sg.FoliageQuality=2
 +CVars=sg.LandscapeQuality=2
 +CVars=r.MobileContentScaleFactor=1.0
++CVars=UI.SlateSDFText.RasterizationMode=Msdf
++CVars=UI.SlateSDFText.ResolutionLevel=2
 
 [Android_Emulator DeviceProfile]
 DeviceType=Android
@@ -1197,6 +1212,8 @@ BaseProfileName=
 bIsVisibleForAssets=True
 +CVars=r.Shaders.ZeroInitialise=1
 +CVars=r.Shaders.BoundsChecking=1
++CVars=UI.SlateSDFText.RasterizationMode=Msdf
++CVars=UI.SlateSDFText.ResolutionLevel=2
 
 [MacClient DeviceProfile]
 DeviceType=Mac
@@ -1216,6 +1233,8 @@ BaseProfileName=
 bIsVisibleForAssets=True
 MeshLODSettings=
 TextureLODSettings=
++CVars=UI.SlateSDFText.RasterizationMode=Msdf
++CVars=UI.SlateSDFText.ResolutionLevel=2
 
 [LinuxEditor DeviceProfile]
 DeviceType=Linux
