Adding unattended retry values for virtualized assets so dependent steps don't fall over if p4 is offline

Submitting on behalf of Derek Ehrman

diff --git a/Engine/Config/BaseEngine.ini b/Engine/Config/BaseEngine.ini
--- a/Engine/Config/BaseEngine.ini
+++ b/Engine/Config/BaseEngine.ini
@@ -1562,8 +1562,8 @@ UseLegacyErrorHandling=False
 ForceCachingOnPull=False
 PullErrorAdditionalMsg=""
 ConnectionHelpUrl=""
-UnattendedRetryCount=0
-UnattendedRetryTimer=0
+UnattendedRetryCount=20
+UnattendedRetryTimer=60
 
 [/Script/Virtualization.VirtualizationFilterSettings]
 !ExcludePackagePaths=ClearArray
