[Backout] - CL36632011 - was submitted to the wrong place
Original CL Desc
-----------------------------------------------------------------
Adding unattended retry values for virtualized assets so dependent steps don't fall over if p4 is offline

diff --git a/Engine/Config/BaseEngine.ini b/Engine/Config/BaseEngine.ini
--- a/Engine/Config/BaseEngine.ini
+++ b/Engine/Config/BaseEngine.ini
@@ -1563,8 +1563,8 @@ UseLegacyErrorHandling=False
 ForceCachingOnPull=False
 PullErrorAdditionalMsg=""
 ConnectionHelpUrl=""
-UnattendedRetryCount=20
-UnattendedRetryTimer=60
+UnattendedRetryCount=0
+UnattendedRetryTimer=0
 
 [/Script/Virtualization.VirtualizationFilterSettings]
 !ExcludePackagePaths=ClearArray
