StudioTelemetry's log analytics is now disabled by default.

- Although we should now have fixed the reported performance issues with FAnalyticsProviderLog it should be an opt in feature rather than an opt out.

diff --git a/Engine/Config/BaseEngine.ini b/Engine/Config/BaseEngine.ini
--- a/Engine/Config/BaseEngine.ini
+++ b/Engine/Config/BaseEngine.ini
@@ -3747,10 +3747,10 @@ SendHardwareData=true
 SendOSData=true
 
 ; Example Log File Analytics
-[StudioTelemetry.Provider.LogAnalytics] 
-Name=LogAnalytics
-ProviderModule=AnalyticsLog
-UsageType=Editor
+;[StudioTelemetry.Provider.LogAnalytics] 
+; Name=LogAnalytics
+; ProviderModule=AnalyticsLog
+; UsageType=Editor
 ; Override the default output location from Saved/Telemetry/Telemetry.json
 ; FileName="somefile.json"
 ; FolderPath="somefolder"
