Fixup redirects and remove UFUNCTION markup for redirected SetAnimClass function

#rnx

diff --git a/Engine/Config/BaseEngine.ini b/Engine/Config/BaseEngine.ini
--- a/Engine/Config/BaseEngine.ini
+++ b/Engine/Config/BaseEngine.ini
@@ -713,10 +713,6 @@ DefaultWorldPartitionHLODLayer=/Engine/Maps/Templates/HLODs/HLODLayer_Instanced.
 
 +PropertyRedirects=(OldName="GameplayStatics.CreatePlayer.bSpawnPawn",NewName="bSpawnPlayerController")
 
-; 4.23
-+FunctionRedirects=(OldName="SkeletalMeshComponent.SetAnimInstanceClass",NewName="SkeletalMeshComponent.SetAnimClass")
-+FunctionRedirects=(OldName="SkeletalMeshComponent.K2_SetAnimInstanceClass",NewName="SkeletalMeshComponent.SetAnimClass")
-
 ; Chaos
 +ClassRedirects=(OldName="/Script/GeometryCollectionCore.GeometryCollection",NewName="/Script/GeometryCollectionEngine.GeometryCollection")
 +ClassRedirects=(OldName="/Script/GeometryCollectionCore.GeometryCollectionCache",NewName="/Script/GeometryCollectionEngine.GeometryCollectionCache")
@@ -727,8 +723,6 @@ DefaultWorldPartitionHLODLayer=/Engine/Maps/Templates/HLODs/HLODLayer_Instanced.
 +FunctionRedirects=(OldName="PlayerController.ClientPlayForceFeedback",NewName="PlayerController.K2_ClientPlayForceFeedback")
 
 ; 4.23
-+FunctionRedirects=(OldName="SkeletalMeshComponent.SetAnimInstanceClass",NewName="SkeletalMeshComponent.SetAnimClass")
-+FunctionRedirects=(OldName="SkeletalMeshComponent.K2_SetAnimInstanceClass",NewName="SkeletalMeshComponent.SetAnimClass")
 +FunctionRedirects=(OldName="EditorUtilityWidget.OnDefaultActionClicked",NewName="EditorUtilityWidget.Run")
 +FunctionRedirects=(OldName="SkeletalMeshComponent.GetSubInstanceByName",NewName="KismetSystemLibrary.GetSubInstanceByTag")
 +StructRedirects=(OldName="/Script/AnimGraphRuntime.AnimNode_Root",NewName="/Script/Engine.AnimNode_Root")
@@ -1335,6 +1329,9 @@ DefaultWorldPartitionHLODLayer=/Engine/Maps/Templates/HLODs/HLODLayer_Instanced.
 +EnumRedirects=(OldName="/Script/GameplayCameras.ECameraAnimationPlaySpace", NewName="/Script/EngineCameras.ECameraAnimationPlaySpace")
 +EnumRedirects=(OldName="/Script/GameplayCameras.ECameraAnimationEasingType", NewName="/Script/EngineCameras.ECameraAnimationEasingType")
 
++FunctionRedirects=(OldName="SkeletalMeshComponent.SetAnimClass",NewName="SetAnimInstanceClass")
++FunctionRedirects=(OldName="SkeletalMeshComponent.K2_SetAnimInstanceClass",NewName="SetAnimInstanceClass")
+
 ; End of releases redirects
 
 ; Fixing the output of tags for blueprints to be reset.
diff --git a/Engine/Source/Runtime/Engine/Classes/Components/SkeletalMeshComponent.h b/Engine/Source/Runtime/Engine/Classes/Components/SkeletalMeshComponent.h
--- a/Engine/Source/Runtime/Engine/Classes/Components/SkeletalMeshComponent.h
+++ b/Engine/Source/Runtime/Engine/Classes/Components/SkeletalMeshComponent.h
@@ -964,7 +964,6 @@ public:
 
 	/** Set the anim instance class. Clears and re-initializes the anim instance with the new class and sets animation mode to 'AnimationBlueprint' */
 	UE_DEPRECATED(5.5, "This function is deprecated. Please use 'SetAnimInstanceClass' instead. ")
-	UFUNCTION(BlueprintCallable, Category = "Components|SkeletalMesh", meta = (Keywords = "AnimBlueprint", DisplayName = "Set Anim Class", DeprecatedFunction, DeprecationMessage = "This function is deprecated. Please use 'SetAnimInstanceClass' instead. "))
 	ENGINE_API virtual void SetAnimClass(class UClass* NewClass);
 
 	/** Get the anim instance class via getter callable by sequencer.  */
