Fix Android Vulkan SM5 renderer

Broken by change GRHISupportsUAVFormatAliasing=true

diff --git a/Engine/Config/BaseDeviceProfiles.ini b/Engine/Config/BaseDeviceProfiles.ini
--- a/Engine/Config/BaseDeviceProfiles.ini
+++ b/Engine/Config/BaseDeviceProfiles.ini
@@ -940,6 +940,7 @@ BaseProfileName=Android
 +CVars=sg.LandscapeQuality=2
 +CVars=r.BloomQuality=2
 +CVars=r.LightShaftQuality=1
++CVars=r.LumenScene.SurfaceCache.Compress=0
 ; Shadows
 +CVars=r.Shadow.MaxResolution=2048
 +CVars=r.Shadow.MaxCSMResolution=2048
