IncrementalCook: Enable the iterativeskip of BodySetups, now that the cook indeterminism in BodySetupGuid has been worked around.
#rnx

diff --git a/Engine/Config/BaseEditor.ini b/Engine/Config/BaseEditor.ini
--- a/Engine/Config/BaseEditor.ini
+++ b/Engine/Config/BaseEditor.ini
@@ -498,7 +498,6 @@ SaveUnversioned=true
 +IterativeClassDenyList=/Script/ComputeFramework.ComputeGraph				; UE-203846
 +IterativeClassDenyList=/Script/CustomizableObject.CustomizableObject		; UE-224744
 +IterativeClassDenyList=/Script/Engine.BlueprintGeneratedClass 				; Not yet documented
-+IterativeClassDenyList=/Script/Engine.BodySetup 							; UE-202602
 +IterativeClassDenyList=/Script/Engine.World 								; Not yet documented
 +IterativeClassDenyList=/Script/PoseSearch.PoseSearchDatabase				; UE-224771
 
