Lumen - optimize surface cache lighting update on High scalability by setting lower update speed. Now with new frustum based prioritization it holds up quite well with those new settings and it makes Lumen Scene Lighting pass x2 faster.

diff --git a/Engine/Config/BaseScalability.ini b/Engine/Config/BaseScalability.ini
--- a/Engine/Config/BaseScalability.ini
+++ b/Engine/Config/BaseScalability.ini
@@ -278,6 +278,9 @@ r.RayTracing.Scene.BuildMode=0
 r.DistanceFieldAO=1
 r.AOQuality=1
 r.Lumen.DiffuseIndirect.Allow=1
+r.LumenScene.DirectLighting.MaxLightsPerTile=4
+r.LumenScene.DirectLighting.UpdateFactor=64
+r.LumenScene.Radiosity.UpdateFactor=128
 r.LumenScene.Radiosity.ProbeSpacing=8
 r.LumenScene.Radiosity.HemisphereProbeResolution=3
 r.Lumen.TraceMeshSDFs.Allow=0
@@ -302,6 +305,9 @@ r.RayTracing.Scene.BuildMode=0
 r.DistanceFieldAO=1
 r.AOQuality=2
 r.Lumen.DiffuseIndirect.Allow=1
+r.LumenScene.DirectLighting.MaxLightsPerTile=8
+r.LumenScene.DirectLighting.UpdateFactor=32
+r.LumenScene.Radiosity.UpdateFactor=64
 r.LumenScene.Radiosity.ProbeSpacing=4
 r.LumenScene.Radiosity.HemisphereProbeResolution=4
 r.Lumen.TraceMeshSDFs.Allow=1
@@ -327,6 +333,9 @@ r.RayTracing.Scene.BuildMode=1
 r.DistanceFieldAO=1
 r.AOQuality=2
 r.Lumen.DiffuseIndirect.Allow=1
+r.LumenScene.DirectLighting.MaxLightsPerTile=8
+r.LumenScene.DirectLighting.UpdateFactor=32
+r.LumenScene.Radiosity.UpdateFactor=64
 r.LumenScene.Radiosity.ProbeSpacing=4
 r.LumenScene.Radiosity.HemisphereProbeResolution=4
 r.Lumen.TraceMeshSDFs.Allow=1
