Add iPadMini7 profile

diff --git a/Engine/Config/BaseDeviceProfiles.ini b/Engine/Config/BaseDeviceProfiles.ini
--- a/Engine/Config/BaseDeviceProfiles.ini
+++ b/Engine/Config/BaseDeviceProfiles.ini
@@ -19,6 +19,7 @@
 +DeviceProfileNameAndTypes=iPadMini4,IOS
 +DeviceProfileNameAndTypes=iPadMini5,IOS
 +DeviceProfileNameAndTypes=iPadMini6,IOS
++DeviceProfileNameAndTypes=iPadMini7,IOS
 +DeviceProfileNameAndTypes=iPodTouch7,IOS
 +DeviceProfileNameAndTypes=iPhone6S,IOS
 +DeviceProfileNameAndTypes=iPhone6SPlus,IOS
@@ -288,6 +289,7 @@ iPad14,[3,4]=iPadPro4_11
 iPad14,[5,6]=iPadPro6_129
 iPad14,[8,9]=iPadAir6
 iPad14,1$=iPadAir_13
+iPad16,[1,2]=iPadMini7
 iPad16,[3,4]=iPadPro5_11
 iPad16,[5,6]=iPadPro7_129
 AppleTV5,=AppleTV
@@ -714,6 +716,10 @@ DeviceType=IOS
 BaseProfileName=IPadPro
 +CVars=ios.PhysicalScreenDensity=326
 
+[iPadMini7 DeviceProfile]
+DeviceType=IOS
+BaseProfileName=iPadMini6
+
 [AppleTV DeviceProfile]
 DeviceType=IOS
 BaseProfileName=IOS
