Fix for Pulsing Lumen AO

#tests it's already live

diff --git a/Engine/Config/BaseScalability.ini b/Engine/Config/BaseScalability.ini
--- a/Engine/Config/BaseScalability.ini
+++ b/Engine/Config/BaseScalability.ini
@@ -292,7 +292,7 @@ r.Lumen.ScreenProbeGather.IrradianceFormat=1
 r.Lumen.ScreenProbeGather.StochasticInterpolation=1
 r.Lumen.ScreenProbeGather.FullResolutionJitterWidth=1
 r.Lumen.ScreenProbeGather.TwoSidedFoliageBackfaceDiffuse=0
-r.Lumen.ScreenProbeGather.ScreenTraces.HZBTraversal.FullResDepth=0
+r.Lumen.ScreenProbeGather.ScreenTraces.HZBTraversal.FullResDepth=1
 r.Lumen.ScreenProbeGather.ShortRangeAO.HardwareRayTracing=0
 r.Lumen.TranslucencyVolume.GridPixelSize=64
 r.Lumen.TranslucencyVolume.TraceFromVolume=0
