Heterogeneous Volumes: extended scalability settings

diff --git a/Engine/Config/BaseScalability.ini b/Engine/Config/BaseScalability.ini
--- a/Engine/Config/BaseScalability.ini
+++ b/Engine/Config/BaseScalability.ini
@@ -660,9 +660,7 @@ r.SkyAtmosphere.TransmittanceLUT.UseSmallFormat=1
 r.SkyAtmosphere.TransmittanceLUT.SampleCount=10.0
 r.SkyAtmosphere.MultiScatteringLUT.SampleCount=15.0
 fx.Niagara.QualityLevel=0
-r.HeterogeneousVolumes.DownsampleFactor=8
-r.HeterogeneousVolumes.MaxStepCount=128
-r.HeterogeneousVolumes.Shadows.Resolution=128
+r.HeterogeneousVolumes=0
 
 [EffectsQuality@1]
 r.TranslucencyLightingVolumeDim=32
@@ -690,9 +688,7 @@ r.SkyAtmosphere.TransmittanceLUT.UseSmallFormat=0
 r.SkyAtmosphere.TransmittanceLUT.SampleCount=10.0
 r.SkyAtmosphere.MultiScatteringLUT.SampleCount=15.0
 fx.Niagara.QualityLevel=1
-r.HeterogeneousVolumes.DownsampleFactor=4
-r.HeterogeneousVolumes.MaxStepCount=256
-r.HeterogeneousVolumes.Shadows.Resolution=128
+r.HeterogeneousVolumes=0
 
 [EffectsQuality@2]
 r.TranslucencyLightingVolumeDim=48
@@ -720,7 +716,7 @@ r.SkyAtmosphere.TransmittanceLUT.UseSmallFormat=0
 r.SkyAtmosphere.TransmittanceLUT.SampleCount=10.0
 r.SkyAtmosphere.MultiScatteringLUT.SampleCount=15.0
 fx.Niagara.QualityLevel=2
-r.HeterogeneousVolumes.DownsampleFactor=2
+r.HeterogeneousVolumes.DownsampleFactor=3
 r.HeterogeneousVolumes.MaxStepCount=256
 r.HeterogeneousVolumes.Shadows.Resolution=256
 
@@ -751,9 +747,10 @@ r.SkyAtmosphere.TransmittanceLUT.SampleCount=10.0
 r.SkyAtmosphere.MultiScatteringLUT.SampleCount=15.0
 fx.Niagara.QualityLevel=3
 r.Refraction.OffsetQuality=1
-r.HeterogeneousVolumes.DownsampleFactor=1
-r.HeterogeneousVolumes.MaxStepCount=512
-r.HeterogeneousVolumes.Shadows.Resolution=512
+r.HeterogeneousVolumes.DownsampleFactor=2
+r.HeterogeneousVolumes.MaxStepCount=256
+r.HeterogeneousVolumes.Shadows.Resolution=256
+r.HeterogeneousVolumes.Shadows.MaxSampleCount=8
 
 [EffectsQuality@Cine]
 r.TranslucencyLightingVolumeDim=64
@@ -785,6 +782,7 @@ r.Refraction.OffsetQuality=1
 r.HeterogeneousVolumes.DownsampleFactor=1
 r.HeterogeneousVolumes.MaxStepCount=1024
 r.HeterogeneousVolumes.Shadows.Resolution=1024
+r.HeterogeneousVolumes.Shadows.MaxSampleCount=16
 
 ;-----------------------------------------------------------------------------------------------------------------
 
diff --git a/Engine/Source/Runtime/Renderer/Private/HeterogeneousVolumes/HeterogeneousVolumes.cpp b/Engine/Source/Runtime/Renderer/Private/HeterogeneousVolumes/HeterogeneousVolumes.cpp
--- a/Engine/Source/Runtime/Renderer/Private/HeterogeneousVolumes/HeterogeneousVolumes.cpp
+++ b/Engine/Source/Runtime/Renderer/Private/HeterogeneousVolumes/HeterogeneousVolumes.cpp
@@ -16,7 +16,7 @@ static TAutoConsoleVariable<int32> CVarHeterogeneousVolumes(
 	TEXT("r.HeterogeneousVolumes"),
 	1,
 	TEXT("Enables the Heterogeneous volume integrator (Default = 1)"),
-	ECVF_RenderThreadSafe
+	ECVF_Scalability | ECVF_RenderThreadSafe
 );
 
 static TAutoConsoleVariable<int32> CVarHeterogeneousVolumesShadows(
@@ -76,7 +76,7 @@ static TAutoConsoleVariable<float> CVarHeterogeneousVolumesIndirectLighting(
 	TEXT("r.HeterogeneousVolumes.IndirectLighting"),
 	0.0f,
 	TEXT("Enables indirect lighting (Default = 0.0)"),
-	ECVF_RenderThreadSafe
+	ECVF_Scalability | ECVF_RenderThreadSafe
 );
 
 static TAutoConsoleVariable<int32> CVarHeterogeneousVolumesStochasticFiltering(
