Allow BPEditorExtension module to load for the Editor. This allows cooking blueprints that depend on classes within.

diff --git a/Engine/Config/BaseGame.ini b/Engine/Config/BaseGame.ini
--- a/Engine/Config/BaseGame.ini
+++ b/Engine/Config/BaseGame.ini
@@ -300,7 +300,6 @@ MapMode=Uncooked
 +DisabledPlugins=Volumetrics
 +DisabledPlugins=AssetSearch
 ; various plugins that cause issues in cooked/shipping editors (unknown reasons for now)
-+DisabledPlugins=BPEditorExtension
 +DisabledPlugins=ObjectBrowser
 ; stuff that doesn't work well or is not needed in shipping
 +DisabledPlugins=FunctionalTestingEditor
