VulkanRHI: Bring back refraction tests now that driver has been upgraded.  Set the minimum Nvidia driver version in Windows to 560.

diff --git a/Engine/Config/BaseHardware.ini b/Engine/Config/BaseHardware.ini
--- a/Engine/Config/BaseHardware.ini
+++ b/Engine/Config/BaseHardware.ini
@@ -34,14 +34,14 @@
 ;-----------------------------------------------------------------------------------------------------------------
 
 [GPU_NVIDIA Windows]
-SuggestedDriverVersion="536.40"
+SuggestedDriverVersion="560.81"
 
 ; The newest driver which supports legacy Nvidia GPUs (700 series and older) is 474.64. We allow this driver for DX11.
 +DriverDenyList=(DriverVersion="<474.64", Reason="These driver versions have known stability issues and missing features", RHI="D3D11")
 
 ; DX12 and Vulkan require recent drivers. The 536 series was released in June 2023.
 +DriverDenyList=(DriverVersion="<536.00", Reason="These driver versions have known stability issues and missing features", RHI="D3D12")
-+DriverDenyList=(DriverVersion="<536.00", Reason="These driver versions have known stability issues and missing features", RHI="Vulkan")
++DriverDenyList=(DriverVersion="<560.00", Reason="These driver versions have known stability issues and missing features", RHI="Vulkan")
 
 ;-----------------------------------------------------------------------------------------------------------------
 
