Add rudimentary support for newer IMG GPUs.

#rnx

diff --git a/Engine/Config/BaseDeviceProfiles.ini b/Engine/Config/BaseDeviceProfiles.ini
--- a/Engine/Config/BaseDeviceProfiles.ini
+++ b/Engine/Config/BaseDeviceProfiles.ini
@@ -825,6 +825,8 @@ BaseProfileName=VisionPro
 +MatchProfile=(Profile="Android_PowerVR_GE8xxx",Match=((SourceType=SRC_GpuFamily,CompareType=CMP_Regex,MatchString="PowerVR Rogue GE8[0-9]+")))
 +MatchProfile=(Profile="Android_PowerVR_GM9xxx_Vulkan",Match=((SourceType=SRC_GpuFamily,CompareType=CMP_Regex,MatchString="PowerVR Rogue GM9[0-9]+"),(SourceType=SRC_VulkanAvailable,CompareType=CMP_Equal,MatchString="true")))
 +MatchProfile=(Profile="Android_PowerVR_GM9xxx",Match=((SourceType=SRC_GpuFamily,CompareType=CMP_Regex,MatchString="PowerVR Rogue GM9[0-9]+")))
++MatchProfile=(Profile="Android_PowerVR_ABCD_Vulkan",Match=((SourceType=SRC_GpuFamily,CompareType=CMP_Regex,MatchString="PowerVR [A-D]-Series"),(SourceType=SRC_VulkanAvailable,CompareType=CMP_Equal,MatchString="true")))
++MatchProfile=(Profile="Android_PowerVR_ABCD",Match=((SourceType=SRC_GpuFamily,CompareType=CMP_Regex,MatchString="PowerVR [A-D]-Series")))
 ; Tegra devices
 
 +MatchProfile=(Profile="Android_TegraK1",Match=((SourceType=SRC_GpuFamily,CompareType=CMP_Equal,MatchString="NVIDIA Tegra"),(SourceType=SRC_GlVersion,CompareType=CMP_Regex,MatchString="^OpenGL ES 3\\.")))
@@ -1082,6 +1084,14 @@ BaseProfileName=Android_PowerVR_GM9xxx
 ; enable Vulkan
 +CVars=r.Android.DisableVulkanSupport=0
 
+[Android_PowerVR_ABCD DeviceProfile]
+DeviceType=Android
+BaseProfileName=Android_PowerVR_GM9xxx
+
+[Android_PowerVR_ABCD_Vulkan DeviceProfile]
+DeviceType=Android
+BaseProfileName=Android_PowerVR_GM9xxx_Vulkan
+
 ;----------------
 ; Mali devices, we support from Mali-T6xx and above
 
