Cook Determinism: Change GlintTexture2Name capitalization to match capitalization in perforce; this mismatch is indeterministically not being created and creating indeterminism in the capitalization inside of the cooked AssetRegistry.
#rnx

diff --git a/Engine/Config/BaseEngine.ini b/Engine/Config/BaseEngine.ini
--- a/Engine/Config/BaseEngine.ini
+++ b/Engine/Config/BaseEngine.ini
@@ -171,7 +171,7 @@ SheenEnergyTextureName=/Engine/EngineMaterials/EnergyConservation/Sheen_Energy.S
 SheenLegacyEnergyTextureName=/Engine/EngineMaterials/EnergyConservation/SheenLegacy_Energy.SheenLegacy_Energy
 DiffuseEnergyTextureName=/Engine/EngineMaterials/EnergyConservation/Diffuse_Energy.Diffuse_Energy
 GlintTextureName=/Engine/EngineMaterials/Substrate/GLints/AGlintTex2dArray.AGlintTex2dArray
-GlintTexture2Name=/Engine/EngineMaterials/Substrate/GLints2/AGlintTex2dArray2.AGlintTex2dArray2
+GlintTexture2Name=/Engine/EngineMaterials/Substrate/Glints2/AGlintTex2dArray2.AGlintTex2dArray2
 SimpleVolumeTextureName=/Engine/EngineMaterials/Substrate/Volume/SubstrateSimpleVolumeLUT3D.SubstrateSimpleVolumeLUT3D
 SimpleVolumeEnvTextureName=/Engine/EngineMaterials/Substrate/Volume/SubstrateSimpleVolumeEnvLUT3D.SubstrateSimpleVolumeEnvLUT3D
 MiniFontTextureName=/Engine/EngineMaterials/MiniFont.MiniFont
