PR #12484: Add the AppStoreConnect ini keys to the base ini key deny list

diff --git a/Engine/Config/BaseGame.ini b/Engine/Config/BaseGame.ini
--- a/Engine/Config/BaseGame.ini
+++ b/Engine/Config/BaseGame.ini
@@ -153,6 +153,9 @@ WriteBackMetadataToAssetRegistry=Disabled
 +IniKeyDenylist=IOSTeamID
 +IniKeyDenylist=SigningCertificate
 +IniKeyDenylist=MobileProvision
++IniKeyDenylist=AppStoreConnectKeyPath
++IniKeyDenylist=AppStoreConnectIssuerID
++IniKeyDenylist=AppStoreConnectKeyID
 +IniKeyDenylist=IniKeyDenylist
 +IniKeyDenylist=IniSectionDenylist
 
