Camera: Added a DisplayName to allow FNamedFilmbackPreset to be localized

diff --git a/Engine/Config/BaseEngine.ini b/Engine/Config/BaseEngine.ini
--- a/Engine/Config/BaseEngine.ini
+++ b/Engine/Config/BaseEngine.ini
@@ -3505,19 +3505,19 @@ LowestAllowedChunkIndexForAutoGeneration=100
 ServerUpdateCameraTimeout=2.0
 
 [/Script/CinematicCamera.CineCameraSettings]
-+FilmbackPresets=(Name="16:9 Film",FilmbackSettings=(SensorWidth=24.00,SensorHeight=13.5))
-+FilmbackPresets=(Name="16:9 Digital Film",FilmbackSettings=(SensorWidth=23.76,SensorHeight=13.365))
-+FilmbackPresets=(Name="16:9 DSLR",FilmbackSettings=(SensorWidth=36,SensorHeight=20.25))
-+FilmbackPresets=(Name="Super 8mm",FilmbackSettings=(SensorWidth=5.79,SensorHeight=4.01))
-+FilmbackPresets=(Name="Super 16mm",FilmbackSettings=(SensorWidth=12.52,SensorHeight=7.58))
-+FilmbackPresets=(Name="Super 35mm",FilmbackSettings=(SensorWidth=24.89,SensorHeight=18.66))
-+FilmbackPresets=(Name="35mm Academy",FilmbackSettings=(SensorWidth=21.946,SensorHeight=16.002))
-+FilmbackPresets=(Name="35mm Full Aperture",FilmbackSettings=(SensorWidth=24.892,SensorHeight=18.9121))
-+FilmbackPresets=(Name="35mm VistaVision",FilmbackSettings=(SensorWidth=37.719,SensorHeight=25.146))
-+FilmbackPresets=(Name="IMAX 70mm",FilmbackSettings=(SensorWidth=70.41,SensorHeight=56.63))
-+FilmbackPresets=(Name="APS-C (Canon)",FilmbackSettings=(SensorWidth=22.2,SensorHeight=14.8))
-+FilmbackPresets=(Name="Full Frame DSLR",FilmbackSettings=(SensorWidth=36,SensorHeight=24))
-+FilmbackPresets=(Name="Micro Four Thirds",FilmbackSettings=(SensorWidth=17.3,SensorHeight=13))
++FilmbackPresets=(DisplayName=LOCTEXT("16:9Film", "16:9 Film"),Name="16:9 Film",FilmbackSettings=(SensorWidth=24.00,SensorHeight=13.5))
++FilmbackPresets=(DisplayName=LOCTEXT("16:9DigitalFilm", "16:9 Digital Film"),Name="16:9 Digital Film",FilmbackSettings=(SensorWidth=23.76,SensorHeight=13.365))
++FilmbackPresets=(DisplayName=LOCTEXT("16:9DSLR", "16:9 DSLR"),Name="16:9 DSLR",FilmbackSettings=(SensorWidth=36,SensorHeight=20.25))
++FilmbackPresets=(DisplayName=LOCTEXT("Super8mm", "Super 8mm"),Name="Super 8mm",FilmbackSettings=(SensorWidth=5.79,SensorHeight=4.01))
++FilmbackPresets=(DisplayName=LOCTEXT("Super16mm", "Super 16mm"),Name="Super 16mm"FilmbackSettings=(SensorWidth=12.52,SensorHeight=7.58))
++FilmbackPresets=(DisplayName=LOCTEXT("Super35mm", "Super 35mm"),Name="Super 35mm",FilmbackSettings=(SensorWidth=24.89,SensorHeight=18.66))
++FilmbackPresets=(DisplayName=LOCTEXT("35mmAcademy", "35mm Academy"),Name="35mm Academy",FilmbackSettings=(SensorWidth=21.946,SensorHeight=16.002))
++FilmbackPresets=(DisplayName=LOCTEXT("35mmFullAperture", "35mm Full Aperture"),Name="35mm Full Aperture",FilmbackSettings=(SensorWidth=24.892,SensorHeight=18.9121))
++FilmbackPresets=(DisplayName=LOCTEXT("35mmVistaVision", "35mm VistaVision"),Name="35mm VistaVision",FilmbackSettings=(SensorWidth=37.719,SensorHeight=25.146))
++FilmbackPresets=(DisplayName=LOCTEXT("IMAX70mm", "IMAX 70mm"),Name="IMAX 70mm",FilmbackSettings=(SensorWidth=70.41,SensorHeight=56.63))
++FilmbackPresets=(DisplayName=LOCTEXT("APS-CCanon", "APS-C (Canon)"),Name="APS-C (Canon)",FilmbackSettings=(SensorWidth=22.2,SensorHeight=14.8))
++FilmbackPresets=(DisplayName=LOCTEXT("FullFrameDSLR", "Full Frame DSLR"),Name="Full Frame DSLR",FilmbackSettings=(SensorWidth=36,SensorHeight=24))
++FilmbackPresets=(DisplayName=LOCTEXT("MicroFourThirds", "Micro Four Thirds"),Name="Micro Four Thirds",FilmbackSettings=(SensorWidth=17.3,SensorHeight=13))
 DefaultFilmbackPreset="16:9 Digital Film"
 +LensPresets=(Name="12mm Prime f/2.8",LensSettings=(MinFocalLength=12,MaxFocalLength=12,MinFStop=2.8,MaxFStop=22,DiaphragmBladeCount=7))
 +LensPresets=(Name="30mm Prime f/1.4",LensSettings=(MinFocalLength=30,MaxFocalLength=30,MinFStop=1.4,MaxFStop=22,DiaphragmBladeCount=7))
diff --git a/Engine/Source/Runtime/CinematicCamera/Private/CineCameraComponent.cpp b/Engine/Source/Runtime/CinematicCamera/Private/CineCameraComponent.cpp
--- a/Engine/Source/Runtime/CinematicCamera/Private/CineCameraComponent.cpp
+++ b/Engine/Source/Runtime/CinematicCamera/Private/CineCameraComponent.cpp
@@ -620,7 +620,7 @@ FText UCineCameraComponent::GetFilmbackText() const
 	{
 		return FText::Format(
 			LOCTEXT("PresetFormat","FilmbackPreset: {0} | Zoom: {1}mm | Av: {2} | Squeeze: {3}"),
-			FText::FromString(Preset->Name),
+			Preset->DisplayName,
 			FText::AsNumber(CurrentFocalLength),
 			FText::AsNumber(CurrentAperture),
 			FText::AsNumber(LensSettings.SqueezeFactor)
diff --git a/Engine/Source/Runtime/CinematicCamera/Public/CineCameraSettings.h b/Engine/Source/Runtime/CinematicCamera/Public/CineCameraSettings.h
--- a/Engine/Source/Runtime/CinematicCamera/Public/CineCameraSettings.h
+++ b/Engine/Source/Runtime/CinematicCamera/Public/CineCameraSettings.h
@@ -74,6 +74,10 @@ struct FNamedFilmbackPreset
 	UPROPERTY(BlueprintReadWrite, EditAnywhere, Category = "Filmback")
 	FString Name;
 
+	/** Display name for the preset. */
+	UPROPERTY(BlueprintReadWrite, EditAnywhere, Category = "Filmback")
+	FText DisplayName;
+
 	UPROPERTY(BlueprintReadWrite, EditAnywhere, Category = "Filmback")
 	FCameraFilmbackSettings FilmbackSettings;
 };
