Fix Android Vulkan SM5 Lumen ray tracing

In Android SM5 preview Lumen raytracing is disabled for now since the shader format configuration used in for preview is incompatible.

diff --git a/Engine/Config/Android/DataDrivenPlatformInfo.ini b/Engine/Config/Android/DataDrivenPlatformInfo.ini
--- a/Engine/Config/Android/DataDrivenPlatformInfo.ini
+++ b/Engine/Config/Android/DataDrivenPlatformInfo.ini
@@ -56,7 +56,7 @@ ActiveIconName=LevelEditor.PreviewMode.AndroidVulkanSM5.Enabled
 InactiveIconName=LevelEditor.PreviewMode.AndroidVulkanSM5.Disabled
 MenuTooltip=NSLOCTEXT("PreviewPlatform", "PreviewMenuTooltip_AndroidVulkanSM5", "Mobile preview using Android Vulkan SM5 quality settings.")
 IconText=NSLOCTEXT("PreviewPlatform", "PreviewIconText_AndroidVulkanSM5", "Android VK SM5")
-DeviceProfileName=Android_Vulkan_SM5
+DeviceProfileName=Android_Preview_Vulkan_SM5
 
 [ShaderPlatform OPENGL_ES3_1_ANDROID]
 Language=OpenGL
diff --git a/Engine/Config/BaseDeviceProfiles.ini b/Engine/Config/BaseDeviceProfiles.ini
--- a/Engine/Config/BaseDeviceProfiles.ini
+++ b/Engine/Config/BaseDeviceProfiles.ini
@@ -100,6 +100,7 @@
 +DeviceProfileNameAndTypes=Android,Android
 +DeviceProfileNameAndTypes=Android_Preview_OpenGL,Android
 +DeviceProfileNameAndTypes=Android_Preview_Vulkan,Android
++DeviceProfileNameAndTypes=Android_Preview_Vulkan_SM5,Android
 +DeviceProfileNameAndTypes=Android_Low,Android
 +DeviceProfileNameAndTypes=Android_Mid,Android
 +DeviceProfileNameAndTypes=Android_High,Android
@@ -966,6 +967,11 @@ BaseProfileName=Android
 +CVars=r.Vulkan.RayTracing.AllowCompaction=0
 +CVars=r.Vulkan.RayTracing.TLASPreferFastTraceTLAS=0
 
+[Android_Preview_Vulkan_SM5 DeviceProfile]
+DeviceType=Android
+BaseProfileName=Android
++CVars=r.Lumen.HardwareRayTracing=False
+
 [Android_Adreno_Vulkan_SM5 DeviceProfile]
 DeviceType=Android
 BaseProfileName=Android_Vulkan_SM5
diff --git a/Engine/Source/Runtime/RenderCore/Private/RenderUtils.cpp b/Engine/Source/Runtime/RenderCore/Private/RenderUtils.cpp
--- a/Engine/Source/Runtime/RenderCore/Private/RenderUtils.cpp
+++ b/Engine/Source/Runtime/RenderCore/Private/RenderUtils.cpp
@@ -1676,7 +1676,7 @@ bool IsRayTracingAllowed()
 
 bool IsRayTracingEnabled(EShaderPlatform ShaderPlatform)
 {
-	return IsRayTracingEnabled() && RHISupportsRayTracing(ShaderPlatform) && IsFeatureLevelSupported(ShaderPlatform, ERHIFeatureLevel::SM6);
+	return IsRayTracingEnabled() && RHISupportsRayTracing(ShaderPlatform) && (IsFeatureLevelSupported(ShaderPlatform, ERHIFeatureLevel::SM6) || IsVulkanMobileSM5Platform(ShaderPlatform));
 }
 
 ERayTracingMode GetRayTracingMode()
