Remove Shipping entitlement to allow fallback to premade.

diff --git a/Engine/Config/BaseEngine.ini b/Engine/Config/BaseEngine.ini
--- a/Engine/Config/BaseEngine.ini
+++ b/Engine/Config/BaseEngine.ini
@@ -3386,7 +3386,6 @@ TemplateIOSPlist=(FilePath="/Game/Build/IOS/UBTGenerated/Info.Template.plist")
 PremadeMacEntitlements=(FilePath="/Game/Build/Mac/Resources/Sandbox.Server.entitlements")
 ShippingSpecificMacEntitlements=(FilePath="/Game/Build/Mac/Resources/Sandbox.NoNet.entitlements")
 PremadeIOSEntitlements=(FilePath="/Game/Build/IOS/Resources/Development.entitlements")
-ShippingSpecificIOSEntitlements=(FilePath="/Game/Build/IOS/Resources/Shipping.entitlements")
 ; Programs generally don't need entitlements. Override in a particular program's .ini to true to use the above Premade entitlements files
 bUseEntitlementsForPrograms=false
 ; The following are used only with Manual codesigning
