Fix namespace on localization text for the menu categories for MaterialExpressions

- "MaterialExpressions" should be "MaterialExpression"

diff --git a/Engine/Config/BaseEditor.ini b/Engine/Config/BaseEditor.ini
--- a/Engine/Config/BaseEditor.ini
+++ b/Engine/Config/BaseEditor.ini
@@ -710,5 +710,5 @@ bUseWidgetTemplateSelector=True
 +Filters=(FilterString="FColorVertexBuffer",DisplayName="Color Vertex Buffers")
 
 [/Script/Engine.MaterialExpressionDesaturation]
-+MenuCategories=NSLOCTEXT("MaterialExpressions", "Color", "Color")
-+MenuCategories=NSLOCTEXT("MaterialExpressions", "Utility", "Utility")
\ No newline at end of file
++MenuCategories=NSLOCTEXT("MaterialExpression", "Color", "Color")
++MenuCategories=NSLOCTEXT("MaterialExpression", "Utility", "Utility")
\ No newline at end of file
