Remove PoseSearchDatabase from the iterative cook class exclusion list to ALLOW it to cook incrementally

diff --git a/Engine/Config/BaseEditor.ini b/Engine/Config/BaseEditor.ini
--- a/Engine/Config/BaseEditor.ini
+++ b/Engine/Config/BaseEditor.ini
@@ -500,7 +500,6 @@ SaveUnversioned=true
 +IterativeClassDenyList=/Script/CustomizableObject.CustomizableObject		; UE-224744
 +IterativeClassDenyList=/Script/Engine.BlueprintGeneratedClass 				; Not yet documented
 +IterativeClassDenyList=/Script/Engine.World 								; Not yet documented
-+IterativeClassDenyList=/Script/PoseSearch.PoseSearchDatabase				; UE-224771
 
 ; Settings for the IterativeValidate cooker mode, which reports false positive iterative skips in incremental cooks
 [IterativeValidate]
