Making the default for editor -game use synchronous registry scanning again due to issues with Lyra assuming the the scan will be completed synchronously.

diff --git a/Engine/Config/BaseEngine.ini b/Engine/Config/BaseEngine.ini
--- a/Engine/Config/BaseEngine.ini
+++ b/Engine/Config/BaseEngine.ini
@@ -2911,9 +2911,9 @@ bFilterSearchableNames=true
 ; the same way the editor does to have the same AssetRegistry available that the cooked game has. This adds cputime and
 ; memory cost to the editor game startup, so some projects may want to turn it off for their game when running editor -game.
 ; Set EditorGameScansAR=false to turn it off.
-; Set EditorGameScansAR=true to turn it on (default). Will use async scanning for -game and, when perf threads are enabled, -server
-; Set EditorGameScansAR=sync to enable scans but run them synchronously
-; Set EditorGameScansAR=async (same as 'true')
+; Set EditorGameScansAR=true to turn it on (default). Will use synchronous scanning for -game and -server.
+; Set EditorGameScansAR=sync to enable scans but run them synchronously (same as 'true')
+; Set EditorGameScansAR=async (same as 'true') Will use async scanning for -game and, when perf threads are enabled, -server
 EditorGameScansAR=true
 
 ; Performance settings for the AssetRegistryScan
diff --git a/Engine/Source/Runtime/AssetRegistry/Private/AssetRegistry.cpp b/Engine/Source/Runtime/AssetRegistry/Private/AssetRegistry.cpp
--- a/Engine/Source/Runtime/AssetRegistry/Private/AssetRegistry.cpp
+++ b/Engine/Source/Runtime/AssetRegistry/Private/AssetRegistry.cpp
@@ -233,7 +233,7 @@ EEditorGameScanMode GetEditorGameScanMode()
 	}
 	else if (EditorGameScansAR.Equals(TEXT("True"), ESearchCase::IgnoreCase))
 	{
-		Result = EEditorGameScanMode::Async;
+		Result = EEditorGameScanMode::Sync;
 	}
 	else
 	{
