[Materials] Initial work on data-driven external code declarations for new material translator.

This is the initial change to move the material translator (including the legacy FHLSLMaterialTranslator) towards a data-driven approach for all external HLSL code injections.
FMaterialExternalCodeDeclaration represents the main structure to organize external HLSL code expressions.
This structure is managed by the singleton MaterialExternalCodeRegistry and referenced by material expressions that extent UMaterialExpressionExternalCodeBase.
In FHLSLMaterialTranslator, they can also be invoked by their name, for material expressions that partially construct their output out of custom HLSL code.
In the upcoming changes, material expression classes that solely inject external HLSL code will be successively replaced by the new UMaterialExpressionExternalCodeBase class.

diff --git a/Engine/Config/BaseEditor.ini b/Engine/Config/BaseEditor.ini
--- a/Engine/Config/BaseEditor.ini
+++ b/Engine/Config/BaseEditor.ini
@@ -708,941 +708,3 @@ bUseWidgetTemplateSelector=True
 +Filters=(FilterString="TexCoordRHIBuffer",DisplayName="Tex Coord Buffers")
 +Filters=(FilterString="FColorVertexBuffer",DisplayName="Color Vertex Buffers")
 
-[/Script/Engine.MaterialExpressionBentNormalCustomOutput]
-+MenuCategories=NSLOCTEXT("MaterialExpression", "Utility", "Utility")
-
-[/Script/Engine.MaterialExpressionClearCoatNormalCustomOutput]
-+MenuCategories=NSLOCTEXT("MaterialExpression", "Utility", "Utility")
-
-[/Script/Engine.MaterialExpressionTangentOutput]
-+MenuCategories=NSLOCTEXT("MaterialExpression", "Custom", "Custom")
-
-[/Script/Engine.MaterialExpressionThinTranslucentMaterialOutput]
-+MenuCategories=NSLOCTEXT("MaterialExpression", "ThinTranslucent", "ThinTranslucent")
-
-[/Script/Engine.MaterialExpressionFirstPersonOutput]
-+MenuCategories=NSLOCTEXT("MaterialExpression", "FirstPerson", "First Person")
-
-[/Script/Landscape.MaterialExpressionLandscapeGrassOutput]
-+MenuCategories=NSLOCTEXT("Landscape", "Landscape", "Landscape")
-
-[/Script/Landscape.MaterialExpressionLandscapeLayerBlend]
-+MenuCategories=NSLOCTEXT("Landscape", "Landscape", "Landscape")
-
-[/Script/Landscape.MaterialExpressionLandscapeLayerCoords]
-+MenuCategories=NSLOCTEXT("Landscape", "Landscape", "Landscape")
-
-[/Script/Landscape.MaterialExpressionLandscapeLayerSample]
-+MenuCategories=NSLOCTEXT("Landscape", "Landscape", "Landscape")
-
-[/Script/Landscape.MaterialExpressionLandscapeLayerSwitch]
-+MenuCategories=NSLOCTEXT("Landscape", "Landscape", "Landscape")
-
-[/Script/Landscape.MaterialExpressionLandscapeLayerWeight]
-+MenuCategories=NSLOCTEXT("Landscape", "Landscape", "Landscape")
-
-[/Script/Landscape.MaterialExpressionLandscapePhysicalMaterialOutput]
-+MenuCategories=NSLOCTEXT("Landscape", "Landscape", "Landscape")
-
-[/Script/Landscape.MaterialExpressionLandscapeVisibilityMask]
-+MenuCategories=NSLOCTEXT("Landscape", "Landscape", "Landscape")
-
-[/Script/Engine.MaterialExpressionAbs]
-+MenuCategories=NSLOCTEXT("MaterialExpression", "Math", "Math")
-
-[/Script/Engine.MaterialExpressionAbsorptionMediumMaterialOutput]
-+MenuCategories=NSLOCTEXT("MaterialExpression", "Shading", "Shading")
-
-[/Script/Engine.MaterialExpressionActorPositionWS]
-+MenuCategories=NSLOCTEXT("MaterialExpression", "Vectors", "Vectors")
-+MenuCategories=NSLOCTEXT("MaterialExpression", "Coordinates", "Coordinates")
-bShaderInputData=true
-
-[/Script/Engine.MaterialExpressionAdd]
-+MenuCategories=NSLOCTEXT("MaterialExpression", "Math", "Math")
-
-[/Script/Engine.MaterialExpressionTextureSample]
-+MenuCategories=NSLOCTEXT("MaterialExpression", "Texture", "Texture")
-
-[/Script/Engine.MaterialExpressionTextureSampleParameter2D]
-+MenuCategories=NSLOCTEXT("MaterialExpression", "Texture", "Texture")
-+MenuCategories=NSLOCTEXT("MaterialExpression", "Parameters", "Parameters")
-
-[/Script/Engine.MaterialExpressionAntialiasedTextureMask]
-+MenuCategories=NSLOCTEXT("MaterialExpression", "Utility", "Utility")
-
-[/Script/Engine.MaterialExpressionAppendVector]
-+MenuCategories=NSLOCTEXT("MaterialExpression", "Math", "Math")
-+MenuCategories=NSLOCTEXT("MaterialExpression", "VectorOps", "VectorOps")
-
-[/Script/Engine.MaterialExpressionArccosine]
-+MenuCategories=NSLOCTEXT("MaterialExpression", "Math", "Math")
-
-[/Script/Engine.MaterialExpressionArccosineFast]
-+MenuCategories=NSLOCTEXT("MaterialExpression", "Math", "Math")
-
-[/Script/Engine.MaterialExpressionArcsine]
-+MenuCategories=NSLOCTEXT("MaterialExpression", "Math", "Math")
-
-[/Script/Engine.MaterialExpressionArcsineFast]
-+MenuCategories=NSLOCTEXT("MaterialExpression", "Math", "Math")
-
-[/Script/Engine.MaterialExpressionArctangent]
-+MenuCategories=NSLOCTEXT("MaterialExpression", "Math", "Math")
-
-[/Script/Engine.MaterialExpressionArctangent2]
-+MenuCategories=NSLOCTEXT("MaterialExpression", "Math", "Math")
-
-[/Script/Engine.MaterialExpressionArctangent2Fast]
-+MenuCategories=NSLOCTEXT("MaterialExpression", "Math", "Math")
-
-[/Script/Engine.MaterialExpressionArctangentFast]
-+MenuCategories=NSLOCTEXT("MaterialExpression", "Math", "Math")
-
-[/Script/Engine.MaterialExpressionAtmosphericFogColor]
-+MenuCategories=NSLOCTEXT("MaterialExpression", "Atmosphere", "Atmosphere")
-
-[/Script/Engine.MaterialExpressionAtmosphericLightColor]
-+MenuCategories=NSLOCTEXT("MaterialExpression", "Utility", "Utility")
-
-[/Script/Engine.MaterialExpressionAtmosphericLightVector]
-+MenuCategories=NSLOCTEXT("MaterialExpression", "Utility", "Utility")
-bShaderInputData=true
-
-[/Script/Engine.MaterialExpressionBindlessSwitch]
-+MenuCategories=NSLOCTEXT("MaterialExpression", "Utility", "Utility")
-
-[/Script/Engine.MaterialExpressionBlackBody]
-+MenuCategories=NSLOCTEXT("MaterialExpression", "Utility", "Utility")
-
-[/Script/Engine.MaterialExpressionBlendMaterialAttributes]
-+MenuCategories=NSLOCTEXT("MaterialExpression", "MaterialAttributes", "Material Attributes")
-
-[/Script/Engine.MaterialExpressionLegacyBlendMaterialAttributes]
-+MenuCategories=NSLOCTEXT("MaterialExpression", "MaterialAttributes", "Material Attributes")
-
-[/Script/Engine.MaterialExpressionBounds]
-+MenuCategories=NSLOCTEXT("MaterialExpression", "Vectors", "Vectors")
-bShaderInputData=true
-
-[/Script/Engine.MaterialExpressionBreakMaterialAttributes]
-+MenuCategories=NSLOCTEXT("MaterialExpression", "MaterialAttributes", "Material Attributes")
-
-[/Script/Engine.MaterialExpressionBumpOffset]
-+MenuCategories=NSLOCTEXT("MaterialExpression", "Utility", "Utility")
-
-[/Script/Engine.MaterialExpressionCameraPositionWS]
-+MenuCategories=NSLOCTEXT("MaterialExpression", "Vectors", "Vectors")
-+MenuCategories=NSLOCTEXT("MaterialExpression", "Coordinates", "Coordinates")
-bShaderInputData=true
-
-[/Script/Engine.MaterialExpressionCameraVectorWS]
-+MenuCategories=NSLOCTEXT("MaterialExpression", "Vectors", "Vectors")
-bShaderInputData=true
-
-[/Script/Engine.MaterialExpressionCeil]
-+MenuCategories=NSLOCTEXT("MaterialExpression", "Math", "Math")
-
-[/Script/Engine.MaterialExpressionComment]
-+MenuCategories=NSLOCTEXT("MaterialExpression", "Utility", "Utility")
-
-[/Script/Engine.MaterialExpressionVectorParameter]
-+MenuCategories=NSLOCTEXT("MaterialExpression", "Parameters", "Parameters")
-
-[/Script/Engine.MaterialExpressionChannelMaskParameter]
-+MenuCategories=NSLOCTEXT("MaterialExpression", "Parameters", "Parameters")
-+MenuCategories=NSLOCTEXT("MaterialExpression", "Utility", "Utility")
-
-[/Script/Engine.MaterialExpressionClamp]
-+MenuCategories=NSLOCTEXT("MaterialExpression", "Math", "Math")
-
-[/Script/Engine.MaterialExpressionCloudSampleAttribute]
-+MenuCategories=NSLOCTEXT("MaterialExpression", "Volume", "Volume")
-
-[/Script/Engine.MaterialExpressionCollectionParameter]
-+MenuCategories=NSLOCTEXT("MaterialExpression", "Parameters", "Parameters")
-
-[/Script/Engine.MaterialExpressionCollectionTransform]
-+MenuCategories=NSLOCTEXT("MaterialExpression", "Parameters", "Parameters")
-
-[/Script/Engine.MaterialExpressionComponentMask]
-+MenuCategories=NSLOCTEXT("MaterialExpression", "Math", "Math")
-+MenuCategories=NSLOCTEXT("MaterialExpression", "VectorOps", "VectorOps")
-
-[/Script/Engine.MaterialExpressionComposite]
-+MenuCategories=NSLOCTEXT("MaterialExpression", "Utility", "Utility")
-
-[/Script/Engine.MaterialExpressionConstant]
-+MenuCategories=NSLOCTEXT("MaterialExpression", "Constants", "Constants")
-
-[/Script/Engine.MaterialExpressionConstant2Vector]
-+MenuCategories=NSLOCTEXT("MaterialExpression", "Constants", "Constants")
-+MenuCategories=NSLOCTEXT("MaterialExpression", "Vectors", "Vectors")
-
-[/Script/Engine.MaterialExpressionConstant3Vector]
-+MenuCategories=NSLOCTEXT("MaterialExpression", "Constants", "Constants")
-+MenuCategories=NSLOCTEXT("MaterialExpression", "Vectors", "Vectors")
-
-[/Script/Engine.MaterialExpressionConstant4Vector]
-+MenuCategories=NSLOCTEXT("MaterialExpression", "Constants", "Constants")
-+MenuCategories=NSLOCTEXT("MaterialExpression", "Vectors", "Vectors")
-
-[/Script/Engine.MaterialExpressionConstantBiasScale]
-+MenuCategories=NSLOCTEXT("MaterialExpression", "Utility", "Utility")
-
-[/Script/Engine.MaterialExpressionCosine]
-+MenuCategories=NSLOCTEXT("MaterialExpression", "Math", "Math")
-
-[/Script/Engine.MaterialExpressionCrossProduct]
-+MenuCategories=NSLOCTEXT("MaterialExpression", "Math", "Math")
-+MenuCategories=NSLOCTEXT("MaterialExpression", "VectorOps", "VectorOps")
-
-[/Script/Engine.MaterialExpressionScalarParameter]
-+MenuCategories=NSLOCTEXT("MaterialExpression", "Parameters", "Parameters")
-
-[/Script/Engine.MaterialExpressionCurveAtlasRowParameter]
-+MenuCategories=NSLOCTEXT("MaterialExpression", "Parameters", "Parameters")
-
-[/Script/Engine.MaterialExpressionCustom]
-+MenuCategories=NSLOCTEXT("MaterialExpression", "Custom", "Custom")
-
-[/Script/Engine.MaterialExpressionDataDrivenShaderPlatformInfoSwitch]
-+MenuCategories=NSLOCTEXT("MaterialExpression", "Utility", "Utility")
-
-[/Script/Engine.MaterialExpressionDBufferTexture]
-+MenuCategories=NSLOCTEXT("MaterialExpression", "Decals", "Decals")
-bShaderInputData=true
-
-[/Script/Engine.MaterialExpressionDDX]
-+MenuCategories=NSLOCTEXT("MaterialExpression", "Utility", "Utility")
-
-[/Script/Engine.MaterialExpressionDDY]
-+MenuCategories=NSLOCTEXT("MaterialExpression", "Utility", "Utility")
-
-[/Script/Engine.MaterialExpressionDecalColor]
-+MenuCategories=NSLOCTEXT("MaterialExpression", "Utils", "Utils")
-bShaderInputData=true
-
-[/Script/Engine.MaterialExpressionDecalDerivative]
-+MenuCategories=NSLOCTEXT("MaterialExpression", "Utils", "Utils")
-bShaderInputData=true
-
-[/Script/Engine.MaterialExpressionDecalLifetimeOpacity]
-+MenuCategories=NSLOCTEXT("MaterialExpression", "Utils", "Utils")
-bShaderInputData=true
-
-[/Script/Engine.MaterialExpressionDecalMipmapLevel]
-+MenuCategories=NSLOCTEXT("MaterialExpression", "Utils", "Utils")
-
-[/Script/Engine.MaterialExpressionDeltaTime]
-+MenuCategories=NSLOCTEXT("MaterialExpression", "Constants", "Constants")
-bShaderInputData=true
-
-[/Script/Engine.MaterialExpressionDepthFade]
-+MenuCategories=NSLOCTEXT("MaterialExpression", "Depth", "Depth")
-+MenuCategories=NSLOCTEXT("MaterialExpression", "Utility", "Utility")
-
-[/Script/Engine.MaterialExpressionDepthOfFieldFunction]
-+MenuCategories=NSLOCTEXT("MaterialExpression", "Utility", "Utility")
-
-[/Script/Engine.MaterialExpressionDeriveNormalZ]
-+MenuCategories=NSLOCTEXT("MaterialExpression", "VectorOps", "VectorOps")
-
-[/Script/Engine.MaterialExpressionDesaturation]
-+MenuCategories=NSLOCTEXT("MaterialExpression", "Color", "Color")
-+MenuCategories=NSLOCTEXT("MaterialExpression", "Utility", "Utility")
-
-[/Script/Engine.MaterialExpressionDistance]
-+MenuCategories=NSLOCTEXT("MaterialExpression", "Utility", "Utility")
-
-[/Script/Engine.MaterialExpressionDistanceCullFade]
-+MenuCategories=NSLOCTEXT("MaterialExpression", "Constants", "Constants")
-bShaderInputData=true
-
-[/Script/Engine.MaterialExpressionDistanceFieldApproxAO]
-+MenuCategories=NSLOCTEXT("MaterialExpression", "Utility", "Utility")
-
-[/Script/Engine.MaterialExpressionDistanceFieldGradient]
-+MenuCategories=NSLOCTEXT("MaterialExpression", "Utility", "Utility")
-
-[/Script/Engine.MaterialExpressionDistanceFieldsRenderingSwitch]
-+MenuCategories=NSLOCTEXT("MaterialExpression", "DistanceFieldsRendering", "DistanceFieldsRendering")
-
-[/Script/Engine.MaterialExpressionDistanceToNearestSurface]
-+MenuCategories=NSLOCTEXT("MaterialExpression", "Utility", "Utility")
-
-[/Script/Engine.MaterialExpressionDivide]
-+MenuCategories=NSLOCTEXT("MaterialExpression", "Math", "Math")
-
-[/Script/Engine.MaterialExpressionDotProduct]
-+MenuCategories=NSLOCTEXT("MaterialExpression", "Math", "Math")
-+MenuCategories=NSLOCTEXT("MaterialExpression", "VectorOps", "VectorOps")
-
-[/Script/Engine.MaterialExpressionDoubleVectorParameter]
-+MenuCategories=NSLOCTEXT("MaterialExpression", "Parameters", "Parameters")
-
-[/Script/Engine.MaterialExpressionDynamicParameter]
-+MenuCategories=NSLOCTEXT("MaterialExpression", "Particles", "Particles")
-+MenuCategories=NSLOCTEXT("MaterialExpression", "Parameters", "Parameters")
-bShaderInputData=true
-
-[/Script/Engine.MaterialExpressionExponential]
-+MenuCategories=NSLOCTEXT("MaterialExpression", "Math", "Math")
-
-[/Script/Engine.MaterialExpressionExponential2]
-+MenuCategories=NSLOCTEXT("MaterialExpression", "Math", "Math")
-
-[/Script/Engine.MaterialExpressionEyeAdaptation]
-+MenuCategories=NSLOCTEXT("MaterialExpression", "Utility", "Utility")
-bShaderInputData=true
-
-[/Script/Engine.MaterialExpressionEyeAdaptationInverse]
-+MenuCategories=NSLOCTEXT("MaterialExpression", "Utility", "Utility")
-bShaderInputData=true
-
-[/Script/Engine.MaterialExpressionFeatureLevelSwitch]
-+MenuCategories=NSLOCTEXT("MaterialExpression", "Utility", "Utility")
-
-[/Script/Engine.MaterialExpressionFloatToUInt]
-+MenuCategories=NSLOCTEXT("MaterialExpression", "Math", "Math")
-
-[/Script/Engine.MaterialExpressionUIntToFloat]
-+MenuCategories=NSLOCTEXT("MaterialExpression", "Math", "Math")
-
-[/Script/Engine.MaterialExpressionFloor]
-+MenuCategories=NSLOCTEXT("MaterialExpression", "Math", "Math")
-
-[/Script/Engine.MaterialExpressionFmod]
-+MenuCategories=NSLOCTEXT("MaterialExpression", "Math", "Math")
-
-[/Script/Engine.MaterialExpressionFontSample]
-+MenuCategories=NSLOCTEXT("MaterialExpression", "Font", "Font")
-+MenuCategories=NSLOCTEXT("MaterialExpression", "Texture", "Texture")
-
-[/Script/Engine.MaterialExpressionFontSampleParameter]
-+MenuCategories=NSLOCTEXT("MaterialExpression", "Font", "Font")
-+MenuCategories=NSLOCTEXT("MaterialExpression", "Texture", "Texture")
-+MenuCategories=NSLOCTEXT("MaterialExpression", "Font", "Font")
-+MenuCategories=NSLOCTEXT("MaterialExpression", "Parameters", "Parameters")
-
-[/Script/Engine.MaterialExpressionFontSignedDistance]
-+MenuCategories=NSLOCTEXT("MaterialExpression", "Constants", "Constants")
-
-[/Script/Engine.MaterialExpressionFrac]
-+MenuCategories=NSLOCTEXT("MaterialExpression", "Math", "Math")
-
-[/Script/Engine.MaterialExpressionFresnel]
-+MenuCategories=NSLOCTEXT("MaterialExpression", "Utility", "Utility")
-
-[/Script/Engine.MaterialExpressionFunctionInput]
-+MenuCategories=NSLOCTEXT("MaterialExpression", "Functions", "Functions")
-
-[/Script/Engine.MaterialExpressionFunctionOutput]
-+MenuCategories=NSLOCTEXT("MaterialExpression", "Functions", "Functions")
-
-[/Script/Engine.MaterialExpressionConstantDouble]
-+MenuCategories=NSLOCTEXT("MaterialExpression", "Constants", "Constants")
-
-[/Script/Engine.MaterialExpressionGetMaterialAttributes]
-+MenuCategories=NSLOCTEXT("MaterialExpression", "MaterialAttributes", "Material Attributes")
-
-[/Script/Engine.MaterialExpressionGIReplace]
-+MenuCategories=NSLOCTEXT("MaterialExpression", "Utility", "Utility")
-
-[/Script/Engine.MaterialExpressionHairAttributes]
-+MenuCategories=NSLOCTEXT("MaterialExpression", "Hair Attributes", "Hair Attributes")
-
-[/Script/Engine.MaterialExpressionHairColor]
-+MenuCategories=NSLOCTEXT("MaterialExpression", "Hair Color", "Hair Color")
-
-[/Script/Engine.MaterialExpressionHsvToRgb]
-+MenuCategories=NSLOCTEXT("MaterialExpression", "Math", "Math")
-
-[/Script/Engine.MaterialExpressionIf]
-+MenuCategories=NSLOCTEXT("MaterialExpression", "Math", "Math")
-
-[/Script/Engine.MaterialExpressionIfThenElse]
-+MenuCategories=NSLOCTEXT("MaterialExpression", "Execution", "Execution")
-
-[/Script/Engine.MaterialExpressionInverseLinearInterpolate]
-+MenuCategories=NSLOCTEXT("MaterialExpression", "Math", "Math")
-+MenuCategories=NSLOCTEXT("MaterialExpression", "Utility", "Utility")
-
-[/Script/Engine.MaterialExpressionIsFirstPerson]
-+MenuCategories=NSLOCTEXT("MaterialExpression", "Constants", "Constants")
-+MenuCategories=NSLOCTEXT("MaterialExpression", "FirstPerson", "First Person")
-bShaderInputData=true
-
-[/Script/Engine.MaterialExpressionIsOrthographic]
-+MenuCategories=NSLOCTEXT("MaterialExpression", "Constants", "Constants")
-bShaderInputData=true
-
-[/Script/Engine.MaterialExpressionLength]
-+MenuCategories=NSLOCTEXT("MaterialExpression", "Utility", "Utility")
-
-[/Script/Engine.MaterialExpressionLightmapUVs]
-+MenuCategories=NSLOCTEXT("MaterialExpression", "Coordinates", "Coordinates")
-bShaderInputData=true
-
-[/Script/Engine.MaterialExpressionLightmassReplace]
-+MenuCategories=NSLOCTEXT("MaterialExpression", "Utility", "Utility")
-
-[/Script/Engine.MaterialExpressionLightVector]
-+MenuCategories=NSLOCTEXT("MaterialExpression", "Vectors", "Vectors")
-
-[/Script/Engine.MaterialExpressionLinearInterpolate]
-+MenuCategories=NSLOCTEXT("MaterialExpression", "Math", "Math")
-+MenuCategories=NSLOCTEXT("MaterialExpression", "Utility", "Utility")
-
-[/Script/Engine.MaterialExpressionLocalPosition]
-+MenuCategories=NSLOCTEXT("MaterialExpression", "Coordinates", "Coordinates")
-bShaderInputData=true
-
-[/Script/Engine.MaterialExpressionLogarithm]
-+MenuCategories=NSLOCTEXT("MaterialExpression", "Math", "Math")
-
-[/Script/Engine.MaterialExpressionLogarithm10]
-+MenuCategories=NSLOCTEXT("MaterialExpression", "Math", "Math")
-
-[/Script/Engine.MaterialExpressionLogarithm2]
-+MenuCategories=NSLOCTEXT("MaterialExpression", "Math", "Math")
-
-[/Script/Engine.MaterialExpressionMakeMaterialAttributes]
-+MenuCategories=NSLOCTEXT("MaterialExpression", "MaterialAttributes", "Material Attributes")
-
-[/Script/Engine.MaterialExpressionMapARPassthroughCameraUV]
-+MenuCategories=NSLOCTEXT("MaterialExpression", "Coordinates", "Coordinates")
-
-[/Script/Engine.MaterialExpressionMaterialAttributeLayers]
-+MenuCategories=NSLOCTEXT("MaterialExpression", "MaterialAttributes", "Material Attributes")
-+MenuCategories=NSLOCTEXT("MaterialExpression", "Parameters", "Parameters")
-
-[/Script/Engine.MaterialExpressionMaterialFunctionCall]
-+MenuCategories=NSLOCTEXT("MaterialExpression", "Functions", "Functions")
-
-[/Script/Engine.MaterialExpressionMaterialProxyReplace]
-+MenuCategories=NSLOCTEXT("MaterialExpression", "Utility", "Utility")
-
-[/Script/Engine.MaterialExpressionMax]
-+MenuCategories=NSLOCTEXT("MaterialExpression", "Math", "Math")
-
-[/Script/Engine.MaterialExpressionMeshPaintTextureCoordinateIndex]
-+MenuCategories=NSLOCTEXT("MaterialExpression", "MeshPaintTexture", "MeshPaintTexture")
-
-[/Script/Engine.MaterialExpressionMeshPaintTextureObject]
-+MenuCategories=NSLOCTEXT("MaterialExpression", "MeshPaintTexture", "MeshPaintTexture")
-
-[/Script/Engine.MaterialExpressionMeshPaintTextureReplace]
-+MenuCategories=NSLOCTEXT("MaterialExpression", "MeshPaintTexture", "MeshPaintTexture")
-
-[/Script/Engine.MaterialExpressionMin]
-+MenuCategories=NSLOCTEXT("MaterialExpression", "Math", "Math")
-
-[/Script/Engine.MaterialExpressionModulo]
-+MenuCategories=NSLOCTEXT("MaterialExpression", "Math", "Math")
-
-[/Script/Engine.MaterialExpressionMultiply]
-+MenuCategories=NSLOCTEXT("MaterialExpression", "Math", "Math")
-
-[/Script/Engine.MaterialExpressionNamedRerouteDeclaration]
-+MenuCategories=NSLOCTEXT("MaterialExpression", "Utility", "Utility")
-
-[/Script/Engine.MaterialExpressionNaniteReplace]
-+MenuCategories=NSLOCTEXT("MaterialExpression", "Utility", "Utility")
-
-[/Script/Engine.MaterialExpressionNeuralNetworkInput]
-+MenuCategories=NSLOCTEXT("MaterialExpression", "NeuralNetwork", "NeuralNetwork")
-
-[/Script/Engine.MaterialExpressionNeuralNetworkOutput]
-+MenuCategories=NSLOCTEXT("MaterialExpression", "NeuralNetwork", "NeuralNetwork")
-bShaderInputData=true
-
-[/Script/Engine.MaterialExpressionNoise]
-+MenuCategories=NSLOCTEXT("MaterialExpressionNoise", "Utility", "Utility")
-
-[/Script/Engine.MaterialExpressionScalarBlueNoise]
-+MenuCategories=NSLOCTEXT("MaterialExpressionNoise", "Noise", "Noise")
-
-[/Script/Engine.MaterialExpressionNormalize]
-+MenuCategories=NSLOCTEXT("MaterialExpression", "Math", "Math")
-+MenuCategories=NSLOCTEXT("MaterialExpression", "VectorOps", "VectorOps")
-
-[/Script/Engine.MaterialExpressionObjectBounds]
-+MenuCategories=NSLOCTEXT("MaterialExpression", "Vectors", "Vectors")
-bShaderInputData=true
-
-[/Script/Engine.MaterialExpressionObjectLocalBounds]
-+MenuCategories=NSLOCTEXT("MaterialExpression", "Vectors", "Vectors")
-bShaderInputData=true
-
-[/Script/Engine.MaterialExpressionObjectOrientation]
-+MenuCategories=NSLOCTEXT("MaterialExpression", "Vectors", "Vectors")
-+MenuCategories=NSLOCTEXT("MaterialExpression", "Coordinates", "Coordinates")
-bShaderInputData=true
-
-[/Script/Engine.MaterialExpressionObjectPositionWS]
-+MenuCategories=NSLOCTEXT("MaterialExpression", "Vectors", "Vectors")
-+MenuCategories=NSLOCTEXT("MaterialExpression", "Coordinates", "Coordinates")
-bShaderInputData=true
-
-[/Script/Engine.MaterialExpressionObjectRadius]
-+MenuCategories=NSLOCTEXT("MaterialExpression", "Coordinates", "Coordinates")
-bShaderInputData=true
-
-[/Script/Engine.MaterialExpressionOneMinus]
-+MenuCategories=NSLOCTEXT("MaterialExpression", "Math", "Math")
-
-[/Script/Engine.MaterialExpressionPanner]
-+MenuCategories=NSLOCTEXT("MaterialExpression", "Coordinates", "Coordinates")
-
-[/Script/Engine.MaterialExpressionParticleColor]
-+MenuCategories=NSLOCTEXT("MaterialExpression", "Particles", "Particles")
-+MenuCategories=NSLOCTEXT("MaterialExpression", "Constants", "Constants")
-bShaderInputData=true
-
-[/Script/Engine.MaterialExpressionParticleDirection]
-+MenuCategories=NSLOCTEXT("MaterialExpression", "Particles", "Particles")
-+MenuCategories=NSLOCTEXT("MaterialExpression", "Constants", "Constants")
-bShaderInputData=true
-
-[/Script/Engine.MaterialExpressionParticleMacroUV]
-+MenuCategories=NSLOCTEXT("MaterialExpression", "Particles", "Particles")
-bShaderInputData=true
-
-[/Script/Engine.MaterialExpressionParticleMotionBlurFade]
-+MenuCategories=NSLOCTEXT("MaterialExpression", "Particles", "Particles")
-+MenuCategories=NSLOCTEXT("MaterialExpression", "Constants", "Constants")
-bShaderInputData=true
-
-[/Script/Engine.MaterialExpressionParticlePositionWS]
-+MenuCategories=NSLOCTEXT("MaterialExpression", "Particles", "Particles")
-+MenuCategories=NSLOCTEXT("MaterialExpression", "Coordinates", "Coordinates")
-+MenuCategories=NSLOCTEXT("MaterialExpression", "Vectors", "Vectors")
-bShaderInputData=true
-
-[/Script/Engine.MaterialExpressionParticleRadius]
-+MenuCategories=NSLOCTEXT("MaterialExpression", "Particles", "Particles")
-+MenuCategories=NSLOCTEXT("MaterialExpression", "Constants", "Constants")
-bShaderInputData=true
-
-[/Script/Engine.MaterialExpressionParticleRandom]
-+MenuCategories=NSLOCTEXT("MaterialExpression", "Particles", "Particles")
-+MenuCategories=NSLOCTEXT("MaterialExpression", "Constants", "Constants")
-bShaderInputData=true
-
-[/Script/Engine.MaterialExpressionParticleRelativeTime]
-+MenuCategories=NSLOCTEXT("MaterialExpression", "Particles", "Particles")
-+MenuCategories=NSLOCTEXT("MaterialExpression", "Constants", "Constants")
-bShaderInputData=true
-
-[/Script/Engine.MaterialExpressionParticleSize]
-+MenuCategories=NSLOCTEXT("MaterialExpression", "Particles", "Particles")
-+MenuCategories=NSLOCTEXT("MaterialExpression", "Constants", "Constants")
-bShaderInputData=true
-
-[/Script/Engine.MaterialExpressionParticleSpeed]
-+MenuCategories=NSLOCTEXT("MaterialExpression", "Particles", "Particles")
-+MenuCategories=NSLOCTEXT("MaterialExpression", "Constants", "Constants")
-bShaderInputData=true
-
-[/Script/Engine.MaterialExpressionParticleSpriteRotation]
-+MenuCategories=NSLOCTEXT("MaterialExpression", "Particles", "Particles")
-+MenuCategories=NSLOCTEXT("MaterialExpression", "Constants", "Constants")
-bShaderInputData=true
-
-[/Script/Engine.MaterialExpressionParticleSubUV]
-+MenuCategories=NSLOCTEXT("MaterialExpression", "Texture", "Texture")
-+MenuCategories=NSLOCTEXT("MaterialExpression", "Particles", "Particles")
-
-[/Script/Engine.MaterialExpressionParticleSubUVProperties]
-+MenuCategories=NSLOCTEXT("MaterialExpression", "Particles", "Particles")
-+MenuCategories=NSLOCTEXT("MaterialExpression", "Coordinates", "Coordinates")
-bShaderInputData=true
-
-[/Script/Engine.MaterialExpressionPathTracingBufferTexture]
-+MenuCategories=NSLOCTEXT("MaterialExpression", "PathTracing", "PathTracing")
-bShaderInputData=true
-
-[/Script/Engine.MaterialExpressionPathTracingQualitySwitch]
-+MenuCategories=NSLOCTEXT("MaterialExpression", "Utility", "Utility")
-
-[/Script/Engine.MaterialExpressionPathTracingRayTypeSwitch]
-+MenuCategories=NSLOCTEXT("MaterialExpression", "Utility", "Utility")
-
-[/Script/Engine.MaterialExpressionPerInstanceCustomData]
-+MenuCategories=NSLOCTEXT("MaterialExpression", "Custom", "Custom")
-bShaderInputData=true
-
-[/Script/Engine.MaterialExpressionPerInstanceCustomData3Vector]
-+MenuCategories=NSLOCTEXT("MaterialExpression", "Custom", "Custom")
-bShaderInputData=true
-
-[/Script/Engine.MaterialExpressionPerInstanceFadeAmount]
-+MenuCategories=NSLOCTEXT("MaterialExpression", "Constants", "Constants")
-bShaderInputData=true
-
-[/Script/Engine.MaterialExpressionPerInstanceRandom]
-+MenuCategories=NSLOCTEXT("MaterialExpression", "Constants", "Constants")
-bShaderInputData=true
-
-[/Script/Engine.MaterialExpressionPinBase]
-+MenuCategories=NSLOCTEXT("MaterialExpression", "Utility", "Utility")
-
-[/Script/Engine.MaterialExpressionPixelDepth]
-+MenuCategories=NSLOCTEXT("MaterialExpression", "Depth", "Depth")
-bShaderInputData=true
-
-[/Script/Engine.MaterialExpressionPixelNormalWS]
-+MenuCategories=NSLOCTEXT("MaterialExpression", "Vectors", "Vectors")
-+MenuCategories=NSLOCTEXT("MaterialExpression", "Coordinates", "Coordinates")
-bShaderInputData=true
-
-[/Script/Engine.MaterialExpressionPostVolumeUserFlagTest]
-+MenuCategories=NSLOCTEXT("MaterialExpression", "Get Post Process Setting", "Get Post Process Setting")
-
-[/Script/Engine.MaterialExpressionPower]
-+MenuCategories=NSLOCTEXT("MaterialExpression", "Math", "Math")
-
-[/Script/Engine.MaterialExpressionPrecomputedAOMask]
-+MenuCategories=NSLOCTEXT("MaterialExpression", "Constants", "Constants")
-bShaderInputData=true
-
-[/Script/Engine.MaterialExpressionPreSkinnedLocalBounds]
-+MenuCategories=NSLOCTEXT("MaterialExpression", "Vectors", "Vectors")
-bShaderInputData=true
-
-[/Script/Engine.MaterialExpressionPreSkinnedNormal]
-+MenuCategories=NSLOCTEXT("MaterialExpression", "Vectors", "Vectors")
-bShaderInputData=true
-
-[/Script/Engine.MaterialExpressionPreSkinnedPosition]
-+MenuCategories=NSLOCTEXT("MaterialExpression", "Vectors", "Vectors")
-bShaderInputData=true
-
-[/Script/Engine.MaterialExpressionPreviousFrameSwitch]
-+MenuCategories=NSLOCTEXT("MaterialExpression", "Functions", "Functions")
-
-[/Script/Engine.MaterialExpressionQualitySwitch]
-+MenuCategories=NSLOCTEXT("MaterialExpression", "Utility", "Utility")
-
-[/Script/Engine.MaterialExpressionRayTracingQualitySwitch]
-+MenuCategories=NSLOCTEXT("MaterialExpression", "Utility", "Utility")
-
-[/Script/Engine.MaterialExpressionReflectionCapturePassSwitch]
-+MenuCategories=NSLOCTEXT("MaterialExpression", "Utility", "Utility")
-
-[/Script/Engine.MaterialExpressionReflectionVectorWS]
-+MenuCategories=NSLOCTEXT("MaterialExpression", "Vectors", "Vectors")
-bShaderInputData=true
-
-[/Script/Engine.MaterialExpressionReroute]
-+MenuCategories=NSLOCTEXT("MaterialExpression", "Utility", "Utility")
-
-[/Script/Engine.MaterialExpressionRequiredSamplersSwitch]
-+MenuCategories=NSLOCTEXT("MaterialExpression", "Utility", "Utility")
-
-[/Script/Engine.MaterialExpressionRgbToHsv]
-+MenuCategories=NSLOCTEXT("MaterialExpression", "Math", "Math")
-
-[/Script/Engine.MaterialExpressionRotateAboutAxis]
-+MenuCategories=NSLOCTEXT("MaterialExpression", "Utility", "Utility")
-
-[/Script/Engine.MaterialExpressionRotator]
-+MenuCategories=NSLOCTEXT("MaterialExpression", "Coordinates", "Coordinates")
-
-[/Script/Engine.MaterialExpressionRound]
-+MenuCategories=NSLOCTEXT("MaterialExpression", "Math", "Math")
-
-[/Script/Engine.MaterialExpressionRuntimeVirtualTextureCustomData]
-+MenuCategories=NSLOCTEXT("MaterialExpression", "VirtualTexture", "VirtualTexture")
-
-[/Script/Engine.MaterialExpressionRuntimeVirtualTextureOutput]
-+MenuCategories=NSLOCTEXT("MaterialExpression", "VirtualTexture", "VirtualTexture")
-
-[/Script/Engine.MaterialExpressionRuntimeVirtualTextureReplace]
-+MenuCategories=NSLOCTEXT("MaterialExpression", "VirtualTexture", "VirtualTexture")
-
-[/Script/Engine.MaterialExpressionRuntimeVirtualTextureSample]
-+MenuCategories=NSLOCTEXT("MaterialExpression", "VirtualTexture", "VirtualTexture")
-
-[/Script/Engine.MaterialExpressionRuntimeVirtualTextureSampleParameter]
-+MenuCategories=NSLOCTEXT("MaterialExpression", "VirtualTexture", "VirtualTexture")
-+MenuCategories=NSLOCTEXT("MaterialExpression", "Parameters", "Parameters")
-
-[/Script/Engine.MaterialExpressionSamplePhysicsVectorField]
-+MenuCategories=NSLOCTEXT("MaterialExpression", "Utility", "Utility")
-
-[/Script/Engine.MaterialExpressionSamplePhysicsScalarField]
-+MenuCategories=NSLOCTEXT("MaterialExpression", "Utility", "Utility")
-
-[/Script/Engine.MaterialExpressionSamplePhysicsIntegerField]
-+MenuCategories=NSLOCTEXT("MaterialExpression", "Utility", "Utility")
-
-[/Script/Engine.MaterialExpressionSaturate]
-+MenuCategories=NSLOCTEXT("MaterialExpression", "Math", "Math")
-
-[/Script/Engine.MaterialExpressionSceneColor]
-+MenuCategories=NSLOCTEXT("MaterialExpression", "Texture", "Texture")
-bShaderInputData=true
-
-[/Script/Engine.MaterialExpressionSceneDepth]
-+MenuCategories=NSLOCTEXT("MaterialExpression", "Depth", "Depth")
-bShaderInputData=true
-
-[/Script/Engine.MaterialExpressionSceneDepthWithoutWater]
-+MenuCategories=NSLOCTEXT("MaterialExpression", "Water", "Water")
-bShaderInputData=true
-
-[/Script/Engine.MaterialExpressionSceneTexelSize]
-+MenuCategories=NSLOCTEXT("MaterialExpression", "Coordinates", "Coordinates")
-bShaderInputData=true
-
-[/Script/Engine.MaterialExpressionSceneTexture]
-+MenuCategories=NSLOCTEXT("MaterialExpression", "Texture", "Texture")
-bShaderInputData=true
-
-[/Script/Engine.MaterialExpressionScreenPosition]
-+MenuCategories=NSLOCTEXT("MaterialExpression", "Coordinates", "Coordinates")
-bShaderInputData=true
-
-[/Script/Engine.MaterialExpressionSetMaterialAttributes]
-+MenuCategories=NSLOCTEXT("MaterialExpression", "MaterialAttributes", "Material Attributes")
-
-[/Script/Engine.MaterialExpressionShaderStageSwitch]
-+MenuCategories=NSLOCTEXT("MaterialExpression", "Utility", "Utility")
-
-[/Script/Engine.MaterialExpressionShadingModel]
-+MenuCategories=NSLOCTEXT("MaterialExpression", "Shading Model", "Shading Model")
-
-[/Script/Engine.MaterialExpressionShadingPathSwitch]
-+MenuCategories=NSLOCTEXT("MaterialExpression", "Utility", "Utility")
-
-[/Script/Engine.MaterialExpressionShadowReplace]
-+MenuCategories=NSLOCTEXT("MaterialExpression", "Utility", "Utility")
-
-[/Script/Engine.MaterialExpressionSign]
-+MenuCategories=NSLOCTEXT("MaterialExpression", "Math", "Math")
-
-[/Script/Engine.MaterialExpressionSine]
-+MenuCategories=NSLOCTEXT("MaterialExpression", "Math", "Math")
-
-[/Script/Engine.MaterialExpressionSingleLayerWaterMaterialOutput]
-+MenuCategories=NSLOCTEXT("MaterialExpression", "Water", "Water")
-
-[/Script/Engine.MaterialExpressionSkyAtmosphereLightDirection]
-+MenuCategories=NSLOCTEXT("MaterialExpression", "Sky", "Sky")
-
-[/Script/Engine.MaterialExpressionSkyAtmosphereLightIlluminance]
-+MenuCategories=NSLOCTEXT("MaterialExpression", "Sky", "Sky")
-
-[/Script/Engine.MaterialExpressionSkyAtmosphereLightIlluminanceOnGround]
-+MenuCategories=NSLOCTEXT("MaterialExpression", "Sky", "Sky")
-
-[/Script/Engine.MaterialExpressionSkyAtmosphereLightDiskLuminance]
-+MenuCategories=NSLOCTEXT("MaterialExpression", "Sky", "Sky")
-
-[/Script/Engine.MaterialExpressionSkyAtmosphereAerialPerspective]
-+MenuCategories=NSLOCTEXT("MaterialExpression", "Sky", "Sky")
-
-[/Script/Engine.MaterialExpressionSkyAtmosphereDistantLightScatteredLuminance]
-+MenuCategories=NSLOCTEXT("MaterialExpression", "Sky", "Sky")
-
-[/Script/Engine.MaterialExpressionSkyAtmosphereViewLuminance]
-+MenuCategories=NSLOCTEXT("MaterialExpression", "Sky", "Sky")
-
-[/Script/Engine.MaterialExpressionSkyLightEnvMapSample]
-+MenuCategories=NSLOCTEXT("MaterialExpression", "Sky", "Sky")
-
-[/Script/Engine.MaterialExpressionSobol]
-+MenuCategories=NSLOCTEXT("MaterialExpression", "Utility", "Utility")
-
-[/Script/Engine.MaterialExpressionSmoothStep]
-+MenuCategories=NSLOCTEXT("MaterialExpression", "Math", "Math")
-+MenuCategories=NSLOCTEXT("MaterialExpression", "Utility", "Utility")
-
-[/Script/Engine.MaterialExpressionSparseVolumeTextureObject]
-+MenuCategories=NSLOCTEXT("MaterialExpression", "Texture", "Texture")
-+MenuCategories=NSLOCTEXT("MaterialExpression", "Functions", "Functions")
-
-[/Script/Engine.MaterialExpressionSparseVolumeTextureSample]
-+MenuCategories=NSLOCTEXT("MaterialExpression", "Texture", "Texture")
-
-[/Script/Engine.MaterialExpressionSparseVolumeTextureSampleParameter]
-+MenuCategories=NSLOCTEXT("MaterialExpression", "Texture", "Texture")
-+MenuCategories=NSLOCTEXT("MaterialExpression", "Parameters", "Parameters")
-
-[/Script/Engine.MaterialExpressionSparseVolumeTextureObjectParameter]
-+MenuCategories=NSLOCTEXT("MaterialExpression", "Texture", "Texture")
-+MenuCategories=NSLOCTEXT("MaterialExpression", "Parameters", "Parameters")
-
-[/Script/Engine.MaterialExpressionSpeedTree]
-+MenuCategories=NSLOCTEXT("MaterialExpression", "SpeedTree", "SpeedTree")
-
-[/Script/Engine.MaterialExpressionSphereMask]
-+MenuCategories=NSLOCTEXT("MaterialExpression", "Utility", "Utility")
-
-[/Script/Engine.MaterialExpressionSphericalParticleOpacity]
-+MenuCategories=NSLOCTEXT("MaterialExpression", "Particles", "Particles")
-
-[/Script/Engine.MaterialExpressionSquareRoot]
-+MenuCategories=NSLOCTEXT("MaterialExpression", "Math", "Math")
-
-[/Script/Engine.MaterialExpressionSRGBColorToWorkingColorSpace]
-+MenuCategories=NSLOCTEXT("MaterialExpression", "Color", "Color")
-
-[/Script/Engine.MaterialExpressionStaticBool]
-+MenuCategories=NSLOCTEXT("MaterialExpression", "Functions", "Functions")
-bShaderInputData=true
-
-[/Script/Engine.MaterialExpressionStaticBoolParameter]
-+MenuCategories=NSLOCTEXT("MaterialExpression", "Parameters", "Parameters")
-
-[/Script/Engine.MaterialExpressionStaticComponentMaskParameter]
-+MenuCategories=NSLOCTEXT("MaterialExpression", "Parameters", "Parameters")
-+MenuCategories=NSLOCTEXT("MaterialExpression", "Parameters", "Parameters")
-
-[/Script/Engine.MaterialExpressionStaticSwitch]
-+MenuCategories=NSLOCTEXT("MaterialExpression", "Functions", "Functions")
-
-[/Script/Engine.MaterialExpressionStaticSwitchParameter]
-+MenuCategories=NSLOCTEXT("MaterialExpression", "Parameters", "Parameters")
-
-[/Script/Engine.MaterialExpressionStep]
-+MenuCategories=NSLOCTEXT("MaterialExpression", "Math", "Math")
-+MenuCategories=NSLOCTEXT("MaterialExpression", "Utility", "Utility")
-
-[/Script/Engine.MaterialExpressionSubsurfaceMediumMaterialOutput]
-+MenuCategories=NSLOCTEXT("MaterialExpression", "Shading", "Shading")
-
-[/Script/Engine.MaterialExpressionSubtract]
-+MenuCategories=NSLOCTEXT("MaterialExpression", "Math", "Math")
-
-[/Script/Engine.MaterialExpressionSwitch]
-+MenuCategories=NSLOCTEXT("MaterialExpression", "Math", "Math")
-
-[/Script/Engine.MaterialExpressionTangent]
-+MenuCategories=NSLOCTEXT("MaterialExpression", "Math", "Math")
-
-[/Script/Engine.MaterialExpressionTemporalSobol]
-+MenuCategories=NSLOCTEXT("MaterialExpression", "Utility", "Utility")
-
-[/Script/Engine.MaterialExpressionTextureCoordinate]
-+MenuCategories=NSLOCTEXT("MaterialExpression", "Coordinates", "Coordinates")
-
-[/Script/Engine.MaterialExpressionTextureCollection]
-+MenuCategories=NSLOCTEXT("MaterialExpression", "Texture", "Texture")
-
-[/Script/Engine.MaterialExpressionTextureCollectionParameter]
-+MenuCategories=NSLOCTEXT("MaterialExpression", "Texture", "Texture")
-+MenuCategories=NSLOCTEXT("MaterialExpression", "Parameters", "Parameters")
-
-[/Script/Engine.MaterialExpressionTextureObject]
-+MenuCategories=NSLOCTEXT("MaterialExpression", "Texture", "Texture")
-+MenuCategories=NSLOCTEXT("MaterialExpression", "Functions", "Functions")
-
-[/Script/Engine.MaterialExpressionTextureObjectFromCollection]
-+MenuCategories=NSLOCTEXT("MaterialExpression", "Texture", "Texture")
-
-[/Script/Engine.MaterialExpressionTextureObjectParameter]
-+MenuCategories=NSLOCTEXT("MaterialExpression", "Texture", "Texture")
-+MenuCategories=NSLOCTEXT("MaterialExpression", "Parameters", "Parameters")
-
-[/Script/Engine.MaterialExpressionTextureProperty]
-+MenuCategories=NSLOCTEXT("MaterialExpression", "Texture", "Texture")
-bShaderInputData=true
-
-[/Script/Engine.MaterialExpressionTextureSampleParameter2DArray]
-+MenuCategories=NSLOCTEXT("MaterialExpression", "Texture", "Texture")
-+MenuCategories=NSLOCTEXT("MaterialExpression", "Parameters", "Parameters")
-
-[/Script/Engine.MaterialExpressionTextureSampleParameterCube]
-+MenuCategories=NSLOCTEXT("MaterialExpression", "Texture", "Texture")
-+MenuCategories=NSLOCTEXT("MaterialExpression", "Parameters", "Parameters")
-
-[/Script/Engine.MaterialExpressionTextureSampleParameterCubeArray]
-+MenuCategories=NSLOCTEXT("MaterialExpression", "Texture", "Texture")
-+MenuCategories=NSLOCTEXT("MaterialExpression", "Parameters", "Parameters")
-
-[/Script/Engine.MaterialExpressionTextureSampleParameterSubUV]
-+MenuCategories=NSLOCTEXT("MaterialExpression", "Texture", "Texture")
-+MenuCategories=NSLOCTEXT("MaterialExpression", "Parameters", "Parameters")
-+MenuCategories=NSLOCTEXT("MaterialExpression", "Particles", "Particles")
-
-[/Script/Engine.MaterialExpressionTextureSampleParameterVolume]
-+MenuCategories=NSLOCTEXT("MaterialExpression", "Texture", "Texture")
-+MenuCategories=NSLOCTEXT("MaterialExpression", "Parameters", "Parameters")
-
-[/Script/Engine.MaterialExpressionTime]
-+MenuCategories=NSLOCTEXT("MaterialExpression", "Constants", "Constants")
-bShaderInputData=true
-
-[/Script/Engine.MaterialExpressionTransform]
-+MenuCategories=NSLOCTEXT("MaterialExpression", "VectorOps", "VectorOps")
-
-[/Script/Engine.MaterialExpressionTransformPosition]
-+MenuCategories=NSLOCTEXT("MaterialExpression", "VectorOps", "VectorOps")
-
-[/Script/Engine.MaterialExpressionTruncate]
-+MenuCategories=NSLOCTEXT("MaterialExpression", "Math", "Math")
-
-[/Script/Engine.MaterialExpressionTruncateLWC]
-+MenuCategories=NSLOCTEXT("MaterialExpression", "Coordinates", "Coordinates")
-
-[/Script/Engine.MaterialExpressionTwoSidedSign]
-+MenuCategories=NSLOCTEXT("MaterialExpression", "Constants", "Constants")
-bShaderInputData=true
-
-[/Script/Engine.MaterialExpressionUserSceneTexture]
-+MenuCategories=NSLOCTEXT("MaterialExpression", "Texture", "Texture")
-bShaderInputData=true
-
-[/Script/Engine.MaterialExpressionVectorNoise]
-+MenuCategories=NSLOCTEXT("MaterialExpressionNoise", "Utility", "Utility")
-
-[/Script/Engine.MaterialExpressionVertexColor]
-+MenuCategories=NSLOCTEXT("MaterialExpression", "Constants", "Constants")
-bShaderInputData=true
-
-[/Script/Engine.MaterialExpressionVertexInterpolator]
-+MenuCategories=NSLOCTEXT("MaterialExpression", "Utility", "Utility")
-
-[/Script/Engine.MaterialExpressionVertexNormalWS]
-+MenuCategories=NSLOCTEXT("MaterialExpression", "Vectors", "Vectors")
-+MenuCategories=NSLOCTEXT("MaterialExpression", "Coordinates", "Coordinates")
-bShaderInputData=true
-
-[/Script/Engine.MaterialExpressionVertexTangentWS]
-+MenuCategories=NSLOCTEXT("MaterialExpression", "Vectors", "Vectors")
-+MenuCategories=NSLOCTEXT("MaterialExpression", "Coordinates", "Coordinates")
-bShaderInputData=true
-
-[/Script/Engine.MaterialExpressionViewProperty]
-+MenuCategories=NSLOCTEXT("MaterialExpression", "Constants", "Constants")
-bShaderInputData=true
-
-[/Script/Engine.MaterialExpressionViewSize]
-+MenuCategories=NSLOCTEXT("MaterialExpression", "Coordinates", "Coordinates")
-bShaderInputData=true
-
-[/Script/Engine.MaterialExpressionVirtualTextureFeatureSwitch]
-+MenuCategories=NSLOCTEXT("MaterialExpression", "VirtualTexture", "VirtualTexture")
-
-[/Script/Engine.MaterialExpressionVolumetricAdvancedMaterialInput]
-+MenuCategories=NSLOCTEXT("MaterialExpression", "Volume", "Volume")
-
-[/Script/Engine.MaterialExpressionVolumetricCloudEmptySpaceSkippingInput]
-+MenuCategories=NSLOCTEXT("MaterialExpression", "Volume", "Volume")
-
-[/Script/Engine.MaterialExpressionVolumetricAdvancedMaterialOutput]
-+MenuCategories=NSLOCTEXT("MaterialExpression", "Volume", "Volume")
-
-[/Script/Engine.MaterialExpressionVolumetricCloudEmptySpaceSkippingOutput]
-+MenuCategories=NSLOCTEXT("MaterialExpression", "Volume", "Volume")
-
-[/Script/Engine.MaterialExpressionWorldPosition]
-+MenuCategories=NSLOCTEXT("MaterialExpression", "Coordinates", "Coordinates")
-bShaderInputData=true
-
-[/Script/Paper2D.MaterialExpressionSpriteTextureSampler]
-+MenuCategories=NSLOCTEXT("MaterialExpression", "Texture", "Texture")
-+MenuCategories=NSLOCTEXT("MaterialExpression", "Parameters", "Parameters")
\ No newline at end of file
diff --git a/Engine/Config/BaseMaterialExpressions.ini b/Engine/Config/BaseMaterialExpressions.ini
new file mode 100644
--- /dev/null
+++ b/Engine/Config/BaseMaterialExpressions.ini
@@ -0,0 +1,975 @@
+; This file defines the data-driven information for all material expressions
+
+[/Script/Engine.MaterialExpressionBentNormalCustomOutput]
++MenuCategories=NSLOCTEXT("MaterialExpression", "Utility", "Utility")
+
+[/Script/Engine.MaterialExpressionClearCoatNormalCustomOutput]
++MenuCategories=NSLOCTEXT("MaterialExpression", "Utility", "Utility")
+
+[/Script/Engine.MaterialExpressionTangentOutput]
++MenuCategories=NSLOCTEXT("MaterialExpression", "Custom", "Custom")
+
+[/Script/Engine.MaterialExpressionThinTranslucentMaterialOutput]
++MenuCategories=NSLOCTEXT("MaterialExpression", "ThinTranslucent", "ThinTranslucent")
+
+[/Script/Engine.MaterialExpressionFirstPersonOutput]
++MenuCategories=NSLOCTEXT("MaterialExpression", "FirstPerson", "First Person")
+
+[/Script/Landscape.MaterialExpressionLandscapeGrassOutput]
++MenuCategories=NSLOCTEXT("Landscape", "Landscape", "Landscape")
+
+[/Script/Landscape.MaterialExpressionLandscapeLayerBlend]
++MenuCategories=NSLOCTEXT("Landscape", "Landscape", "Landscape")
+
+[/Script/Landscape.MaterialExpressionLandscapeLayerCoords]
++MenuCategories=NSLOCTEXT("Landscape", "Landscape", "Landscape")
+
+[/Script/Landscape.MaterialExpressionLandscapeLayerSample]
++MenuCategories=NSLOCTEXT("Landscape", "Landscape", "Landscape")
+
+[/Script/Landscape.MaterialExpressionLandscapeLayerSwitch]
++MenuCategories=NSLOCTEXT("Landscape", "Landscape", "Landscape")
+
+[/Script/Landscape.MaterialExpressionLandscapeLayerWeight]
++MenuCategories=NSLOCTEXT("Landscape", "Landscape", "Landscape")
+
+[/Script/Landscape.MaterialExpressionLandscapePhysicalMaterialOutput]
++MenuCategories=NSLOCTEXT("Landscape", "Landscape", "Landscape")
+
+[/Script/Landscape.MaterialExpressionLandscapeVisibilityMask]
++MenuCategories=NSLOCTEXT("Landscape", "Landscape", "Landscape")
+
+[/Script/Engine.MaterialExpressionAbs]
++MenuCategories=NSLOCTEXT("MaterialExpression", "Math", "Math")
+
+[/Script/Engine.MaterialExpressionAbsorptionMediumMaterialOutput]
++MenuCategories=NSLOCTEXT("MaterialExpression", "Shading", "Shading")
+
+[/Script/Engine.MaterialExpressionActorPositionWS]
++MenuCategories=NSLOCTEXT("MaterialExpression", "Vectors", "Vectors")
++MenuCategories=NSLOCTEXT("MaterialExpression", "Coordinates", "Coordinates")
+bShaderInputData=true
+
+[/Script/Engine.MaterialExpressionAdd]
++MenuCategories=NSLOCTEXT("MaterialExpression", "Math", "Math")
+
+[/Script/Engine.MaterialExpressionTextureSample]
++MenuCategories=NSLOCTEXT("MaterialExpression", "Texture", "Texture")
+
+[/Script/Engine.MaterialExpressionTextureSampleParameter2D]
++MenuCategories=NSLOCTEXT("MaterialExpression", "Texture", "Texture")
++MenuCategories=NSLOCTEXT("MaterialExpression", "Parameters", "Parameters")
+
+[/Script/Engine.MaterialExpressionAntialiasedTextureMask]
++MenuCategories=NSLOCTEXT("MaterialExpression", "Utility", "Utility")
+
+[/Script/Engine.MaterialExpressionAppendVector]
++MenuCategories=NSLOCTEXT("MaterialExpression", "Math", "Math")
++MenuCategories=NSLOCTEXT("MaterialExpression", "VectorOps", "VectorOps")
+
+[/Script/Engine.MaterialExpressionArccosine]
++MenuCategories=NSLOCTEXT("MaterialExpression", "Math", "Math")
+
+[/Script/Engine.MaterialExpressionArccosineFast]
++MenuCategories=NSLOCTEXT("MaterialExpression", "Math", "Math")
+
+[/Script/Engine.MaterialExpressionArcsine]
++MenuCategories=NSLOCTEXT("MaterialExpression", "Math", "Math")
+
+[/Script/Engine.MaterialExpressionArcsineFast]
++MenuCategories=NSLOCTEXT("MaterialExpression", "Math", "Math")
+
+[/Script/Engine.MaterialExpressionArctangent]
++MenuCategories=NSLOCTEXT("MaterialExpression", "Math", "Math")
+
+[/Script/Engine.MaterialExpressionArctangent2]
++MenuCategories=NSLOCTEXT("MaterialExpression", "Math", "Math")
+
+[/Script/Engine.MaterialExpressionArctangent2Fast]
++MenuCategories=NSLOCTEXT("MaterialExpression", "Math", "Math")
+
+[/Script/Engine.MaterialExpressionArctangentFast]
++MenuCategories=NSLOCTEXT("MaterialExpression", "Math", "Math")
+
+[/Script/Engine.MaterialExpressionAtmosphericFogColor]
++MenuCategories=NSLOCTEXT("MaterialExpression", "Atmosphere", "Atmosphere")
+
+[/Script/Engine.MaterialExpressionAtmosphericLightColor]
++MenuCategories=NSLOCTEXT("MaterialExpression", "Utility", "Utility")
+
+[/Script/Engine.MaterialExpressionAtmosphericLightVector]
++MenuCategories=NSLOCTEXT("MaterialExpression", "Utility", "Utility")
+bShaderInputData=true
+
+[/Script/Engine.MaterialExpressionBindlessSwitch]
++MenuCategories=NSLOCTEXT("MaterialExpression", "Utility", "Utility")
+
+[/Script/Engine.MaterialExpressionBlackBody]
++MenuCategories=NSLOCTEXT("MaterialExpression", "Utility", "Utility")
+
+[/Script/Engine.MaterialExpressionBlendMaterialAttributes]
++MenuCategories=NSLOCTEXT("MaterialExpression", "MaterialAttributes", "Material Attributes")
+
+[/Script/Engine.MaterialExpressionLegacyBlendMaterialAttributes]
++MenuCategories=NSLOCTEXT("MaterialExpression", "MaterialAttributes", "Material Attributes")
+
+[/Script/Engine.MaterialExpressionBounds]
++MenuCategories=NSLOCTEXT("MaterialExpression", "Vectors", "Vectors")
+bShaderInputData=true
+
+[/Script/Engine.MaterialExpressionBreakMaterialAttributes]
++MenuCategories=NSLOCTEXT("MaterialExpression", "MaterialAttributes", "Material Attributes")
+
+[/Script/Engine.MaterialExpressionBumpOffset]
++MenuCategories=NSLOCTEXT("MaterialExpression", "Utility", "Utility")
+
+[/Script/Engine.MaterialExpressionCameraPositionWS]
++MenuCategories=NSLOCTEXT("MaterialExpression", "Vectors", "Vectors")
++MenuCategories=NSLOCTEXT("MaterialExpression", "Coordinates", "Coordinates")
+bShaderInputData=true
+
+[/Script/Engine.MaterialExpressionCameraVectorWS]
++MenuCategories=NSLOCTEXT("MaterialExpression", "Vectors", "Vectors")
+bShaderInputData=true
+
+[/Script/Engine.MaterialExpressionCeil]
++MenuCategories=NSLOCTEXT("MaterialExpression", "Math", "Math")
+
+[/Script/Engine.MaterialExpressionComment]
++MenuCategories=NSLOCTEXT("MaterialExpression", "Utility", "Utility")
+
+[/Script/Engine.MaterialExpressionVectorParameter]
++MenuCategories=NSLOCTEXT("MaterialExpression", "Parameters", "Parameters")
+
+[/Script/Engine.MaterialExpressionChannelMaskParameter]
++MenuCategories=NSLOCTEXT("MaterialExpression", "Parameters", "Parameters")
++MenuCategories=NSLOCTEXT("MaterialExpression", "Utility", "Utility")
+
+[/Script/Engine.MaterialExpressionClamp]
++MenuCategories=NSLOCTEXT("MaterialExpression", "Math", "Math")
+
+[/Script/Engine.MaterialExpressionCloudSampleAttribute]
++MenuCategories=NSLOCTEXT("MaterialExpression", "Volume", "Volume")
+
+[/Script/Engine.MaterialExpressionCollectionParameter]
++MenuCategories=NSLOCTEXT("MaterialExpression", "Parameters", "Parameters")
+
+[/Script/Engine.MaterialExpressionCollectionTransform]
++MenuCategories=NSLOCTEXT("MaterialExpression", "Parameters", "Parameters")
+
+[/Script/Engine.MaterialExpressionComponentMask]
++MenuCategories=NSLOCTEXT("MaterialExpression", "Math", "Math")
++MenuCategories=NSLOCTEXT("MaterialExpression", "VectorOps", "VectorOps")
+
+[/Script/Engine.MaterialExpressionComposite]
++MenuCategories=NSLOCTEXT("MaterialExpression", "Utility", "Utility")
+
+[/Script/Engine.MaterialExpressionConstant]
++MenuCategories=NSLOCTEXT("MaterialExpression", "Constants", "Constants")
+
+[/Script/Engine.MaterialExpressionConstant2Vector]
++MenuCategories=NSLOCTEXT("MaterialExpression", "Constants", "Constants")
++MenuCategories=NSLOCTEXT("MaterialExpression", "Vectors", "Vectors")
+
+[/Script/Engine.MaterialExpressionConstant3Vector]
++MenuCategories=NSLOCTEXT("MaterialExpression", "Constants", "Constants")
++MenuCategories=NSLOCTEXT("MaterialExpression", "Vectors", "Vectors")
+
+[/Script/Engine.MaterialExpressionConstant4Vector]
++MenuCategories=NSLOCTEXT("MaterialExpression", "Constants", "Constants")
++MenuCategories=NSLOCTEXT("MaterialExpression", "Vectors", "Vectors")
+
+[/Script/Engine.MaterialExpressionConstantBiasScale]
++MenuCategories=NSLOCTEXT("MaterialExpression", "Utility", "Utility")
+
+[/Script/Engine.MaterialExpressionCosine]
++MenuCategories=NSLOCTEXT("MaterialExpression", "Math", "Math")
+
+[/Script/Engine.MaterialExpressionCrossProduct]
++MenuCategories=NSLOCTEXT("MaterialExpression", "Math", "Math")
++MenuCategories=NSLOCTEXT("MaterialExpression", "VectorOps", "VectorOps")
+
+[/Script/Engine.MaterialExpressionScalarParameter]
++MenuCategories=NSLOCTEXT("MaterialExpression", "Parameters", "Parameters")
+
+[/Script/Engine.MaterialExpressionCurveAtlasRowParameter]
++MenuCategories=NSLOCTEXT("MaterialExpression", "Parameters", "Parameters")
+
+[/Script/Engine.MaterialExpressionCustom]
++MenuCategories=NSLOCTEXT("MaterialExpression", "Custom", "Custom")
+
+[/Script/Engine.MaterialExpressionDataDrivenShaderPlatformInfoSwitch]
++MenuCategories=NSLOCTEXT("MaterialExpression", "Utility", "Utility")
+
+[/Script/Engine.MaterialExpressionDBufferTexture]
++MenuCategories=NSLOCTEXT("MaterialExpression", "Decals", "Decals")
+bShaderInputData=true
+
+[/Script/Engine.MaterialExpressionDDX]
++MenuCategories=NSLOCTEXT("MaterialExpression", "Utility", "Utility")
+
+[/Script/Engine.MaterialExpressionDDY]
++MenuCategories=NSLOCTEXT("MaterialExpression", "Utility", "Utility")
+
+[/Script/Engine.MaterialExpressionDecalColor]
++MenuCategories=NSLOCTEXT("MaterialExpression", "Utils", "Utils")
+bShaderInputData=true
+
+[/Script/Engine.MaterialExpressionDecalDerivative]
++MenuCategories=NSLOCTEXT("MaterialExpression", "Utils", "Utils")
+bShaderInputData=true
+
+[/Script/Engine.MaterialExpressionDecalLifetimeOpacity]
++MenuCategories=NSLOCTEXT("MaterialExpression", "Utils", "Utils")
+bShaderInputData=true
+
+[/Script/Engine.MaterialExpressionDecalMipmapLevel]
++MenuCategories=NSLOCTEXT("MaterialExpression", "Utils", "Utils")
+
+[/Script/Engine.MaterialExpressionDeltaTime]
++MenuCategories=NSLOCTEXT("MaterialExpression", "Constants", "Constants")
+bShaderInputData=true
+
+[/Script/Engine.MaterialExpressionDepthFade]
++MenuCategories=NSLOCTEXT("MaterialExpression", "Depth", "Depth")
++MenuCategories=NSLOCTEXT("MaterialExpression", "Utility", "Utility")
+
+[/Script/Engine.MaterialExpressionDepthOfFieldFunction]
++MenuCategories=NSLOCTEXT("MaterialExpression", "Utility", "Utility")
+
+[/Script/Engine.MaterialExpressionDeriveNormalZ]
++MenuCategories=NSLOCTEXT("MaterialExpression", "VectorOps", "VectorOps")
+
+[/Script/Engine.MaterialExpressionDesaturation]
++MenuCategories=NSLOCTEXT("MaterialExpression", "Color", "Color")
++MenuCategories=NSLOCTEXT("MaterialExpression", "Utility", "Utility")
+
+[/Script/Engine.MaterialExpressionDistance]
++MenuCategories=NSLOCTEXT("MaterialExpression", "Utility", "Utility")
+
+[/Script/Engine.MaterialExpressionDistanceCullFade]
++MenuCategories=NSLOCTEXT("MaterialExpression", "Constants", "Constants")
+bShaderInputData=true
+
+[/Script/Engine.MaterialExpressionDistanceFieldApproxAO]
++MenuCategories=NSLOCTEXT("MaterialExpression", "Utility", "Utility")
+
+[/Script/Engine.MaterialExpressionDistanceFieldGradient]
++MenuCategories=NSLOCTEXT("MaterialExpression", "Utility", "Utility")
+
+[/Script/Engine.MaterialExpressionDistanceFieldsRenderingSwitch]
++MenuCategories=NSLOCTEXT("MaterialExpression", "DistanceFieldsRendering", "DistanceFieldsRendering")
+
+[/Script/Engine.MaterialExpressionDistanceToNearestSurface]
++MenuCategories=NSLOCTEXT("MaterialExpression", "Utility", "Utility")
+
+[/Script/Engine.MaterialExpressionDivide]
++MenuCategories=NSLOCTEXT("MaterialExpression", "Math", "Math")
+
+[/Script/Engine.MaterialExpressionDotProduct]
++MenuCategories=NSLOCTEXT("MaterialExpression", "Math", "Math")
++MenuCategories=NSLOCTEXT("MaterialExpression", "VectorOps", "VectorOps")
+
+[/Script/Engine.MaterialExpressionDoubleVectorParameter]
++MenuCategories=NSLOCTEXT("MaterialExpression", "Parameters", "Parameters")
+
+[/Script/Engine.MaterialExpressionDynamicParameter]
++MenuCategories=NSLOCTEXT("MaterialExpression", "Particles", "Particles")
++MenuCategories=NSLOCTEXT("MaterialExpression", "Parameters", "Parameters")
+bShaderInputData=true
+
+[/Script/Engine.MaterialExpressionExponential]
++MenuCategories=NSLOCTEXT("MaterialExpression", "Math", "Math")
+
+[/Script/Engine.MaterialExpressionExponential2]
++MenuCategories=NSLOCTEXT("MaterialExpression", "Math", "Math")
+
+;[/Script/Engine.MaterialExpressionExternalCodeBase]
+
+[/Script/Engine.MaterialExpressionEyeAdaptation]
++MenuCategories=NSLOCTEXT("MaterialExpression", "Utility", "Utility")
+bShaderInputData=true
+
+[/Script/Engine.MaterialExpressionEyeAdaptationInverse]
++MenuCategories=NSLOCTEXT("MaterialExpression", "Utility", "Utility")
+bShaderInputData=true
+
+[/Script/Engine.MaterialExpressionFeatureLevelSwitch]
++MenuCategories=NSLOCTEXT("MaterialExpression", "Utility", "Utility")
+
+[/Script/Engine.MaterialExpressionFloatToUInt]
++MenuCategories=NSLOCTEXT("MaterialExpression", "Math", "Math")
+
+[/Script/Engine.MaterialExpressionUIntToFloat]
++MenuCategories=NSLOCTEXT("MaterialExpression", "Math", "Math")
+
+[/Script/Engine.MaterialExpressionFloor]
++MenuCategories=NSLOCTEXT("MaterialExpression", "Math", "Math")
+
+[/Script/Engine.MaterialExpressionFmod]
++MenuCategories=NSLOCTEXT("MaterialExpression", "Math", "Math")
+
+[/Script/Engine.MaterialExpressionFontSample]
++MenuCategories=NSLOCTEXT("MaterialExpression", "Font", "Font")
++MenuCategories=NSLOCTEXT("MaterialExpression", "Texture", "Texture")
+
+[/Script/Engine.MaterialExpressionFontSampleParameter]
++MenuCategories=NSLOCTEXT("MaterialExpression", "Font", "Font")
++MenuCategories=NSLOCTEXT("MaterialExpression", "Texture", "Texture")
++MenuCategories=NSLOCTEXT("MaterialExpression", "Font", "Font")
++MenuCategories=NSLOCTEXT("MaterialExpression", "Parameters", "Parameters")
+
+[/Script/Engine.MaterialExpressionFontSignedDistance]
++MenuCategories=NSLOCTEXT("MaterialExpression", "Constants", "Constants")
+ExternalCodeIdentifier="FontSignedDistanceData"
+Caption="Font Signed Distance"
+
+[/Script/Engine.MaterialExpressionFrac]
++MenuCategories=NSLOCTEXT("MaterialExpression", "Math", "Math")
+
+[/Script/Engine.MaterialExpressionFresnel]
++MenuCategories=NSLOCTEXT("MaterialExpression", "Utility", "Utility")
+
+[/Script/Engine.MaterialExpressionFunctionInput]
++MenuCategories=NSLOCTEXT("MaterialExpression", "Functions", "Functions")
+
+[/Script/Engine.MaterialExpressionFunctionOutput]
++MenuCategories=NSLOCTEXT("MaterialExpression", "Functions", "Functions")
+
+[/Script/Engine.MaterialExpressionConstantDouble]
++MenuCategories=NSLOCTEXT("MaterialExpression", "Constants", "Constants")
+
+[/Script/Engine.MaterialExpressionGetMaterialAttributes]
++MenuCategories=NSLOCTEXT("MaterialExpression", "MaterialAttributes", "Material Attributes")
+
+[/Script/Engine.MaterialExpressionGIReplace]
++MenuCategories=NSLOCTEXT("MaterialExpression", "Utility", "Utility")
+
+[/Script/Engine.MaterialExpressionHairAttributes]
++MenuCategories=NSLOCTEXT("MaterialExpression", "Hair Attributes", "Hair Attributes")
+
+[/Script/Engine.MaterialExpressionHairColor]
++MenuCategories=NSLOCTEXT("MaterialExpression", "Hair Color", "Hair Color")
+
+[/Script/Engine.MaterialExpressionHsvToRgb]
++MenuCategories=NSLOCTEXT("MaterialExpression", "Math", "Math")
+
+[/Script/Engine.MaterialExpressionIf]
++MenuCategories=NSLOCTEXT("MaterialExpression", "Math", "Math")
+
+[/Script/Engine.MaterialExpressionIfThenElse]
++MenuCategories=NSLOCTEXT("MaterialExpression", "Execution", "Execution")
+
+[/Script/Engine.MaterialExpressionInverseLinearInterpolate]
++MenuCategories=NSLOCTEXT("MaterialExpression", "Math", "Math")
++MenuCategories=NSLOCTEXT("MaterialExpression", "Utility", "Utility")
+
+[/Script/Engine.MaterialExpressionIsFirstPerson]
++MenuCategories=NSLOCTEXT("MaterialExpression", "Constants", "Constants")
++MenuCategories=NSLOCTEXT("MaterialExpression", "FirstPerson", "First Person")
+bShaderInputData=true
+
+[/Script/Engine.MaterialExpressionIsOrthographic]
++MenuCategories=NSLOCTEXT("MaterialExpression", "Constants", "Constants")
+ExternalCodeIdentifier="IsOrthographic"
+bShaderInputData=true
+
+[/Script/Engine.MaterialExpressionLength]
++MenuCategories=NSLOCTEXT("MaterialExpression", "Utility", "Utility")
+
+[/Script/Engine.MaterialExpressionLightmapUVs]
++MenuCategories=NSLOCTEXT("MaterialExpression", "Coordinates", "Coordinates")
+bShaderInputData=true
+
+[/Script/Engine.MaterialExpressionLightmassReplace]
++MenuCategories=NSLOCTEXT("MaterialExpression", "Utility", "Utility")
+
+[/Script/Engine.MaterialExpressionLightVector]
++MenuCategories=NSLOCTEXT("MaterialExpression", "Vectors", "Vectors")
+ExternalCodeIdentifier="LightVector"
+
+[/Script/Engine.MaterialExpressionLinearInterpolate]
++MenuCategories=NSLOCTEXT("MaterialExpression", "Math", "Math")
++MenuCategories=NSLOCTEXT("MaterialExpression", "Utility", "Utility")
+
+[/Script/Engine.MaterialExpressionLocalPosition]
++MenuCategories=NSLOCTEXT("MaterialExpression", "Coordinates", "Coordinates")
+bShaderInputData=true
+
+[/Script/Engine.MaterialExpressionLogarithm]
++MenuCategories=NSLOCTEXT("MaterialExpression", "Math", "Math")
+
+[/Script/Engine.MaterialExpressionLogarithm10]
++MenuCategories=NSLOCTEXT("MaterialExpression", "Math", "Math")
+
+[/Script/Engine.MaterialExpressionLogarithm2]
++MenuCategories=NSLOCTEXT("MaterialExpression", "Math", "Math")
+
+[/Script/Engine.MaterialExpressionMakeMaterialAttributes]
++MenuCategories=NSLOCTEXT("MaterialExpression", "MaterialAttributes", "Material Attributes")
+
+[/Script/Engine.MaterialExpressionMapARPassthroughCameraUV]
++MenuCategories=NSLOCTEXT("MaterialExpression", "Coordinates", "Coordinates")
+
+[/Script/Engine.MaterialExpressionMaterialAttributeLayers]
++MenuCategories=NSLOCTEXT("MaterialExpression", "MaterialAttributes", "Material Attributes")
++MenuCategories=NSLOCTEXT("MaterialExpression", "Parameters", "Parameters")
+
+[/Script/Engine.MaterialExpressionMaterialFunctionCall]
++MenuCategories=NSLOCTEXT("MaterialExpression", "Functions", "Functions")
+
+[/Script/Engine.MaterialExpressionMaterialProxyReplace]
++MenuCategories=NSLOCTEXT("MaterialExpression", "Utility", "Utility")
+
+[/Script/Engine.MaterialExpressionMax]
++MenuCategories=NSLOCTEXT("MaterialExpression", "Math", "Math")
+
+[/Script/Engine.MaterialExpressionMeshPaintTextureCoordinateIndex]
++MenuCategories=NSLOCTEXT("MaterialExpression", "MeshPaintTexture", "MeshPaintTexture")
+ExternalCodeIdentifier="MeshPaintTextureCoordinateIndex"
+
+[/Script/Engine.MaterialExpressionMeshPaintTextureObject]
++MenuCategories=NSLOCTEXT("MaterialExpression", "MeshPaintTexture", "MeshPaintTexture")
+ExternalCodeIdentifier="MeshPaintTextureDescriptor"
+
+[/Script/Engine.MaterialExpressionMeshPaintTextureReplace]
++MenuCategories=NSLOCTEXT("MaterialExpression", "MeshPaintTexture", "MeshPaintTexture")
+
+[/Script/Engine.MaterialExpressionMin]
++MenuCategories=NSLOCTEXT("MaterialExpression", "Math", "Math")
+
+[/Script/Engine.MaterialExpressionModulo]
++MenuCategories=NSLOCTEXT("MaterialExpression", "Math", "Math")
+
+[/Script/Engine.MaterialExpressionMultiply]
++MenuCategories=NSLOCTEXT("MaterialExpression", "Math", "Math")
+
+[/Script/Engine.MaterialExpressionNamedRerouteDeclaration]
++MenuCategories=NSLOCTEXT("MaterialExpression", "Utility", "Utility")
+
+[/Script/Engine.MaterialExpressionNaniteReplace]
++MenuCategories=NSLOCTEXT("MaterialExpression", "Utility", "Utility")
+
+[/Script/Engine.MaterialExpressionNeuralNetworkInput]
++MenuCategories=NSLOCTEXT("MaterialExpression", "NeuralNetwork", "NeuralNetwork")
+
+[/Script/Engine.MaterialExpressionNeuralNetworkOutput]
++MenuCategories=NSLOCTEXT("MaterialExpression", "NeuralNetwork", "NeuralNetwork")
+bShaderInputData=true
+
+[/Script/Engine.MaterialExpressionNoise]
++MenuCategories=NSLOCTEXT("MaterialExpressionNoise", "Utility", "Utility")
+
+[/Script/Engine.MaterialExpressionScalarBlueNoise]
++MenuCategories=NSLOCTEXT("MaterialExpressionNoise", "Noise", "Noise")
+
+[/Script/Engine.MaterialExpressionNormalize]
++MenuCategories=NSLOCTEXT("MaterialExpression", "Math", "Math")
++MenuCategories=NSLOCTEXT("MaterialExpression", "VectorOps", "VectorOps")
+
+[/Script/Engine.MaterialExpressionObjectBounds]
++MenuCategories=NSLOCTEXT("MaterialExpression", "Vectors", "Vectors")
+bShaderInputData=true
+
+[/Script/Engine.MaterialExpressionObjectLocalBounds]
++MenuCategories=NSLOCTEXT("MaterialExpression", "Vectors", "Vectors")
+bShaderInputData=true
+
+[/Script/Engine.MaterialExpressionObjectOrientation]
++MenuCategories=NSLOCTEXT("MaterialExpression", "Vectors", "Vectors")
++MenuCategories=NSLOCTEXT("MaterialExpression", "Coordinates", "Coordinates")
+bShaderInputData=true
+
+[/Script/Engine.MaterialExpressionObjectPositionWS]
++MenuCategories=NSLOCTEXT("MaterialExpression", "Vectors", "Vectors")
++MenuCategories=NSLOCTEXT("MaterialExpression", "Coordinates", "Coordinates")
+bShaderInputData=true
+
+[/Script/Engine.MaterialExpressionObjectRadius]
++MenuCategories=NSLOCTEXT("MaterialExpression", "Coordinates", "Coordinates")
+bShaderInputData=true
+
+[/Script/Engine.MaterialExpressionOneMinus]
++MenuCategories=NSLOCTEXT("MaterialExpression", "Math", "Math")
+
+[/Script/Engine.MaterialExpressionPanner]
++MenuCategories=NSLOCTEXT("MaterialExpression", "Coordinates", "Coordinates")
+
+[/Script/Engine.MaterialExpressionParticleColor]
++MenuCategories=NSLOCTEXT("MaterialExpression", "Particles", "Particles")
++MenuCategories=NSLOCTEXT("MaterialExpression", "Constants", "Constants")
+bShaderInputData=true
+
+[/Script/Engine.MaterialExpressionParticleDirection]
++MenuCategories=NSLOCTEXT("MaterialExpression", "Particles", "Particles")
++MenuCategories=NSLOCTEXT("MaterialExpression", "Constants", "Constants")
+bShaderInputData=true
+
+[/Script/Engine.MaterialExpressionParticleMacroUV]
++MenuCategories=NSLOCTEXT("MaterialExpression", "Particles", "Particles")
+ExternalCodeIdentifier="ParticleMacroUV"
+bShaderInputData=true
+
+[/Script/Engine.MaterialExpressionParticleMotionBlurFade]
++MenuCategories=NSLOCTEXT("MaterialExpression", "Particles", "Particles")
++MenuCategories=NSLOCTEXT("MaterialExpression", "Constants", "Constants")
+bShaderInputData=true
+
+[/Script/Engine.MaterialExpressionParticlePositionWS]
++MenuCategories=NSLOCTEXT("MaterialExpression", "Particles", "Particles")
++MenuCategories=NSLOCTEXT("MaterialExpression", "Coordinates", "Coordinates")
++MenuCategories=NSLOCTEXT("MaterialExpression", "Vectors", "Vectors")
+bShaderInputData=true
+
+[/Script/Engine.MaterialExpressionParticleRadius]
++MenuCategories=NSLOCTEXT("MaterialExpression", "Particles", "Particles")
++MenuCategories=NSLOCTEXT("MaterialExpression", "Constants", "Constants")
+bShaderInputData=true
+ExternalCodeIdentifier="ParticleRadius"
+
+[/Script/Engine.MaterialExpressionParticleRandom]
++MenuCategories=NSLOCTEXT("MaterialExpression", "Particles", "Particles")
++MenuCategories=NSLOCTEXT("MaterialExpression", "Constants", "Constants")
+bShaderInputData=true
+
+[/Script/Engine.MaterialExpressionParticleRelativeTime]
++MenuCategories=NSLOCTEXT("MaterialExpression", "Particles", "Particles")
++MenuCategories=NSLOCTEXT("MaterialExpression", "Constants", "Constants")
+bShaderInputData=true
+
+[/Script/Engine.MaterialExpressionParticleSize]
++MenuCategories=NSLOCTEXT("MaterialExpression", "Particles", "Particles")
++MenuCategories=NSLOCTEXT("MaterialExpression", "Constants", "Constants")
+bShaderInputData=true
+
+[/Script/Engine.MaterialExpressionParticleSpeed]
++MenuCategories=NSLOCTEXT("MaterialExpression", "Particles", "Particles")
++MenuCategories=NSLOCTEXT("MaterialExpression", "Constants", "Constants")
+bShaderInputData=true
+
+[/Script/Engine.MaterialExpressionParticleSpriteRotation]
++MenuCategories=NSLOCTEXT("MaterialExpression", "Particles", "Particles")
++MenuCategories=NSLOCTEXT("MaterialExpression", "Constants", "Constants")
+bShaderInputData=true
+
+[/Script/Engine.MaterialExpressionParticleSubUV]
++MenuCategories=NSLOCTEXT("MaterialExpression", "Texture", "Texture")
++MenuCategories=NSLOCTEXT("MaterialExpression", "Particles", "Particles")
+
+[/Script/Engine.MaterialExpressionParticleSubUVProperties]
++MenuCategories=NSLOCTEXT("MaterialExpression", "Particles", "Particles")
++MenuCategories=NSLOCTEXT("MaterialExpression", "Coordinates", "Coordinates")
+bShaderInputData=true
+
+[/Script/Engine.MaterialExpressionPathTracingBufferTexture]
++MenuCategories=NSLOCTEXT("MaterialExpression", "PathTracing", "PathTracing")
+bShaderInputData=true
+
+[/Script/Engine.MaterialExpressionPathTracingQualitySwitch]
++MenuCategories=NSLOCTEXT("MaterialExpression", "Utility", "Utility")
+
+[/Script/Engine.MaterialExpressionPathTracingRayTypeSwitch]
++MenuCategories=NSLOCTEXT("MaterialExpression", "Utility", "Utility")
+
+[/Script/Engine.MaterialExpressionPerInstanceCustomData]
++MenuCategories=NSLOCTEXT("MaterialExpression", "Custom", "Custom")
+bShaderInputData=true
+
+[/Script/Engine.MaterialExpressionPerInstanceCustomData3Vector]
++MenuCategories=NSLOCTEXT("MaterialExpression", "Custom", "Custom")
+bShaderInputData=true
+
+[/Script/Engine.MaterialExpressionPerInstanceFadeAmount]
++MenuCategories=NSLOCTEXT("MaterialExpression", "Constants", "Constants")
+bShaderInputData=true
+
+[/Script/Engine.MaterialExpressionPerInstanceRandom]
++MenuCategories=NSLOCTEXT("MaterialExpression", "Constants", "Constants")
+bShaderInputData=true
+
+[/Script/Engine.MaterialExpressionPinBase]
++MenuCategories=NSLOCTEXT("MaterialExpression", "Utility", "Utility")
+
+[/Script/Engine.MaterialExpressionPixelDepth]
++MenuCategories=NSLOCTEXT("MaterialExpression", "Depth", "Depth")
+bShaderInputData=true
+
+[/Script/Engine.MaterialExpressionPixelNormalWS]
++MenuCategories=NSLOCTEXT("MaterialExpression", "Vectors", "Vectors")
++MenuCategories=NSLOCTEXT("MaterialExpression", "Coordinates", "Coordinates")
+bShaderInputData=true
+
+[/Script/Engine.MaterialExpressionPostVolumeUserFlagTest]
++MenuCategories=NSLOCTEXT("MaterialExpression", "Get Post Process Setting", "Get Post Process Setting")
+
+[/Script/Engine.MaterialExpressionPower]
++MenuCategories=NSLOCTEXT("MaterialExpression", "Math", "Math")
+
+[/Script/Engine.MaterialExpressionPrecomputedAOMask]
++MenuCategories=NSLOCTEXT("MaterialExpression", "Constants", "Constants")
+bShaderInputData=true
+
+[/Script/Engine.MaterialExpressionPreSkinnedLocalBounds]
++MenuCategories=NSLOCTEXT("MaterialExpression", "Vectors", "Vectors")
+bShaderInputData=true
+
+[/Script/Engine.MaterialExpressionPreSkinnedNormal]
++MenuCategories=NSLOCTEXT("MaterialExpression", "Vectors", "Vectors")
+ExternalCodeIdentifier="PreSkinnedNormal"
+bShaderInputData=true
+
+[/Script/Engine.MaterialExpressionPreSkinnedPosition]
++MenuCategories=NSLOCTEXT("MaterialExpression", "Vectors", "Vectors")
+ExternalCodeIdentifier="PreSkinnedPosition"
+bShaderInputData=true
+
+[/Script/Engine.MaterialExpressionPreviousFrameSwitch]
++MenuCategories=NSLOCTEXT("MaterialExpression", "Functions", "Functions")
+
+[/Script/Engine.MaterialExpressionQualitySwitch]
++MenuCategories=NSLOCTEXT("MaterialExpression", "Utility", "Utility")
+
+[/Script/Engine.MaterialExpressionRayTracingQualitySwitch]
++MenuCategories=NSLOCTEXT("MaterialExpression", "Utility", "Utility")
+
+[/Script/Engine.MaterialExpressionReflectionCapturePassSwitch]
++MenuCategories=NSLOCTEXT("MaterialExpression", "Utility", "Utility")
+
+[/Script/Engine.MaterialExpressionReflectionVectorWS]
++MenuCategories=NSLOCTEXT("MaterialExpression", "Vectors", "Vectors")
+bShaderInputData=true
+
+[/Script/Engine.MaterialExpressionReroute]
++MenuCategories=NSLOCTEXT("MaterialExpression", "Utility", "Utility")
+
+[/Script/Engine.MaterialExpressionRequiredSamplersSwitch]
++MenuCategories=NSLOCTEXT("MaterialExpression", "Utility", "Utility")
+
+[/Script/Engine.MaterialExpressionRgbToHsv]
++MenuCategories=NSLOCTEXT("MaterialExpression", "Math", "Math")
+
+[/Script/Engine.MaterialExpressionRotateAboutAxis]
++MenuCategories=NSLOCTEXT("MaterialExpression", "Utility", "Utility")
+
+[/Script/Engine.MaterialExpressionRotator]
++MenuCategories=NSLOCTEXT("MaterialExpression", "Coordinates", "Coordinates")
+
+[/Script/Engine.MaterialExpressionRound]
++MenuCategories=NSLOCTEXT("MaterialExpression", "Math", "Math")
+
+[/Script/Engine.MaterialExpressionRuntimeVirtualTextureCustomData]
++MenuCategories=NSLOCTEXT("MaterialExpression", "VirtualTexture", "VirtualTexture")
+
+[/Script/Engine.MaterialExpressionRuntimeVirtualTextureOutput]
++MenuCategories=NSLOCTEXT("MaterialExpression", "VirtualTexture", "VirtualTexture")
+
+[/Script/Engine.MaterialExpressionRuntimeVirtualTextureReplace]
++MenuCategories=NSLOCTEXT("MaterialExpression", "VirtualTexture", "VirtualTexture")
+
+[/Script/Engine.MaterialExpressionRuntimeVirtualTextureSample]
++MenuCategories=NSLOCTEXT("MaterialExpression", "VirtualTexture", "VirtualTexture")
+
+[/Script/Engine.MaterialExpressionRuntimeVirtualTextureSampleParameter]
++MenuCategories=NSLOCTEXT("MaterialExpression", "VirtualTexture", "VirtualTexture")
++MenuCategories=NSLOCTEXT("MaterialExpression", "Parameters", "Parameters")
+
+[/Script/Engine.MaterialExpressionSamplePhysicsVectorField]
++MenuCategories=NSLOCTEXT("MaterialExpression", "Utility", "Utility")
+
+[/Script/Engine.MaterialExpressionSamplePhysicsScalarField]
++MenuCategories=NSLOCTEXT("MaterialExpression", "Utility", "Utility")
+
+[/Script/Engine.MaterialExpressionSamplePhysicsIntegerField]
++MenuCategories=NSLOCTEXT("MaterialExpression", "Utility", "Utility")
+
+[/Script/Engine.MaterialExpressionSaturate]
++MenuCategories=NSLOCTEXT("MaterialExpression", "Math", "Math")
+
+[/Script/Engine.MaterialExpressionSceneColor]
++MenuCategories=NSLOCTEXT("MaterialExpression", "Texture", "Texture")
+bShaderInputData=true
+
+[/Script/Engine.MaterialExpressionSceneDepth]
++MenuCategories=NSLOCTEXT("MaterialExpression", "Depth", "Depth")
+bShaderInputData=true
+
+[/Script/Engine.MaterialExpressionSceneDepthWithoutWater]
++MenuCategories=NSLOCTEXT("MaterialExpression", "Water", "Water")
+bShaderInputData=true
+
+[/Script/Engine.MaterialExpressionSceneTexelSize]
++MenuCategories=NSLOCTEXT("MaterialExpression", "Coordinates", "Coordinates")
+bShaderInputData=true
+
+[/Script/Engine.MaterialExpressionSceneTexture]
++MenuCategories=NSLOCTEXT("MaterialExpression", "Texture", "Texture")
+bShaderInputData=true
+
+[/Script/Engine.MaterialExpressionScreenPosition]
++MenuCategories=NSLOCTEXT("MaterialExpression", "Coordinates", "Coordinates")
+bShaderInputData=true
+
+[/Script/Engine.MaterialExpressionSetMaterialAttributes]
++MenuCategories=NSLOCTEXT("MaterialExpression", "MaterialAttributes", "Material Attributes")
+
+[/Script/Engine.MaterialExpressionShaderStageSwitch]
++MenuCategories=NSLOCTEXT("MaterialExpression", "Utility", "Utility")
+
+[/Script/Engine.MaterialExpressionShadingModel]
++MenuCategories=NSLOCTEXT("MaterialExpression", "Shading Model", "Shading Model")
+
+[/Script/Engine.MaterialExpressionShadingPathSwitch]
++MenuCategories=NSLOCTEXT("MaterialExpression", "Utility", "Utility")
+
+[/Script/Engine.MaterialExpressionShadowReplace]
++MenuCategories=NSLOCTEXT("MaterialExpression", "Utility", "Utility")
+
+[/Script/Engine.MaterialExpressionSign]
++MenuCategories=NSLOCTEXT("MaterialExpression", "Math", "Math")
+
+[/Script/Engine.MaterialExpressionSine]
++MenuCategories=NSLOCTEXT("MaterialExpression", "Math", "Math")
+
+[/Script/Engine.MaterialExpressionSingleLayerWaterMaterialOutput]
++MenuCategories=NSLOCTEXT("MaterialExpression", "Water", "Water")
+
+[/Script/Engine.MaterialExpressionSkyAtmosphereLightDirection]
++MenuCategories=NSLOCTEXT("MaterialExpression", "Sky", "Sky")
+
+[/Script/Engine.MaterialExpressionSkyAtmosphereLightIlluminance]
++MenuCategories=NSLOCTEXT("MaterialExpression", "Sky", "Sky")
+
+[/Script/Engine.MaterialExpressionSkyAtmosphereLightIlluminanceOnGround]
++MenuCategories=NSLOCTEXT("MaterialExpression", "Sky", "Sky")
+
+[/Script/Engine.MaterialExpressionSkyAtmosphereLightDiskLuminance]
++MenuCategories=NSLOCTEXT("MaterialExpression", "Sky", "Sky")
+
+[/Script/Engine.MaterialExpressionSkyAtmosphereAerialPerspective]
++MenuCategories=NSLOCTEXT("MaterialExpression", "Sky", "Sky")
+
+[/Script/Engine.MaterialExpressionSkyAtmosphereDistantLightScatteredLuminance]
++MenuCategories=NSLOCTEXT("MaterialExpression", "Sky", "Sky")
+
+[/Script/Engine.MaterialExpressionSkyAtmosphereViewLuminance]
++MenuCategories=NSLOCTEXT("MaterialExpression", "Sky", "Sky")
+
+[/Script/Engine.MaterialExpressionSkyLightEnvMapSample]
++MenuCategories=NSLOCTEXT("MaterialExpression", "Sky", "Sky")
+
+[/Script/Engine.MaterialExpressionSobol]
++MenuCategories=NSLOCTEXT("MaterialExpression", "Utility", "Utility")
+
+[/Script/Engine.MaterialExpressionSmoothStep]
++MenuCategories=NSLOCTEXT("MaterialExpression", "Math", "Math")
++MenuCategories=NSLOCTEXT("MaterialExpression", "Utility", "Utility")
+
+[/Script/Engine.MaterialExpressionSparseVolumeTextureObject]
++MenuCategories=NSLOCTEXT("MaterialExpression", "Texture", "Texture")
++MenuCategories=NSLOCTEXT("MaterialExpression", "Functions", "Functions")
+
+[/Script/Engine.MaterialExpressionSparseVolumeTextureSample]
++MenuCategories=NSLOCTEXT("MaterialExpression", "Texture", "Texture")
+
+[/Script/Engine.MaterialExpressionSparseVolumeTextureSampleParameter]
++MenuCategories=NSLOCTEXT("MaterialExpression", "Texture", "Texture")
++MenuCategories=NSLOCTEXT("MaterialExpression", "Parameters", "Parameters")
+
+[/Script/Engine.MaterialExpressionSparseVolumeTextureObjectParameter]
++MenuCategories=NSLOCTEXT("MaterialExpression", "Texture", "Texture")
++MenuCategories=NSLOCTEXT("MaterialExpression", "Parameters", "Parameters")
+
+[/Script/Engine.MaterialExpressionSpeedTree]
++MenuCategories=NSLOCTEXT("MaterialExpression", "SpeedTree", "SpeedTree")
+
+[/Script/Engine.MaterialExpressionSphereMask]
++MenuCategories=NSLOCTEXT("MaterialExpression", "Utility", "Utility")
+
+[/Script/Engine.MaterialExpressionSphericalParticleOpacity]
++MenuCategories=NSLOCTEXT("MaterialExpression", "Particles", "Particles")
+
+[/Script/Engine.MaterialExpressionSquareRoot]
++MenuCategories=NSLOCTEXT("MaterialExpression", "Math", "Math")
+
+[/Script/Engine.MaterialExpressionSRGBColorToWorkingColorSpace]
++MenuCategories=NSLOCTEXT("MaterialExpression", "Color", "Color")
+
+[/Script/Engine.MaterialExpressionStaticBool]
++MenuCategories=NSLOCTEXT("MaterialExpression", "Functions", "Functions")
+bShaderInputData=true
+
+[/Script/Engine.MaterialExpressionStaticBoolParameter]
++MenuCategories=NSLOCTEXT("MaterialExpression", "Parameters", "Parameters")
+
+[/Script/Engine.MaterialExpressionStaticComponentMaskParameter]
++MenuCategories=NSLOCTEXT("MaterialExpression", "Parameters", "Parameters")
++MenuCategories=NSLOCTEXT("MaterialExpression", "Parameters", "Parameters")
+
+[/Script/Engine.MaterialExpressionStaticSwitch]
++MenuCategories=NSLOCTEXT("MaterialExpression", "Functions", "Functions")
+
+[/Script/Engine.MaterialExpressionStaticSwitchParameter]
++MenuCategories=NSLOCTEXT("MaterialExpression", "Parameters", "Parameters")
+
+[/Script/Engine.MaterialExpressionStep]
++MenuCategories=NSLOCTEXT("MaterialExpression", "Math", "Math")
++MenuCategories=NSLOCTEXT("MaterialExpression", "Utility", "Utility")
+
+[/Script/Engine.MaterialExpressionSubsurfaceMediumMaterialOutput]
++MenuCategories=NSLOCTEXT("MaterialExpression", "Shading", "Shading")
+
+[/Script/Engine.MaterialExpressionSubtract]
++MenuCategories=NSLOCTEXT("MaterialExpression", "Math", "Math")
+
+[/Script/Engine.MaterialExpressionSwitch]
++MenuCategories=NSLOCTEXT("MaterialExpression", "Math", "Math")
+
+[/Script/Engine.MaterialExpressionTangent]
++MenuCategories=NSLOCTEXT("MaterialExpression", "Math", "Math")
+
+[/Script/Engine.MaterialExpressionTemporalSobol]
++MenuCategories=NSLOCTEXT("MaterialExpression", "Utility", "Utility")
+
+[/Script/Engine.MaterialExpressionTextureCoordinate]
++MenuCategories=NSLOCTEXT("MaterialExpression", "Coordinates", "Coordinates")
+
+[/Script/Engine.MaterialExpressionTextureCollection]
++MenuCategories=NSLOCTEXT("MaterialExpression", "Texture", "Texture")
+
+[/Script/Engine.MaterialExpressionTextureCollectionParameter]
++MenuCategories=NSLOCTEXT("MaterialExpression", "Texture", "Texture")
++MenuCategories=NSLOCTEXT("MaterialExpression", "Parameters", "Parameters")
+
+[/Script/Engine.MaterialExpressionTextureObject]
++MenuCategories=NSLOCTEXT("MaterialExpression", "Texture", "Texture")
++MenuCategories=NSLOCTEXT("MaterialExpression", "Functions", "Functions")
+
+[/Script/Engine.MaterialExpressionTextureObjectFromCollection]
++MenuCategories=NSLOCTEXT("MaterialExpression", "Texture", "Texture")
+
+[/Script/Engine.MaterialExpressionTextureObjectParameter]
++MenuCategories=NSLOCTEXT("MaterialExpression", "Texture", "Texture")
++MenuCategories=NSLOCTEXT("MaterialExpression", "Parameters", "Parameters")
+
+[/Script/Engine.MaterialExpressionTextureProperty]
++MenuCategories=NSLOCTEXT("MaterialExpression", "Texture", "Texture")
+bShaderInputData=true
+
+[/Script/Engine.MaterialExpressionTextureSampleParameter2DArray]
++MenuCategories=NSLOCTEXT("MaterialExpression", "Texture", "Texture")
++MenuCategories=NSLOCTEXT("MaterialExpression", "Parameters", "Parameters")
+
+[/Script/Engine.MaterialExpressionTextureSampleParameterCube]
++MenuCategories=NSLOCTEXT("MaterialExpression", "Texture", "Texture")
++MenuCategories=NSLOCTEXT("MaterialExpression", "Parameters", "Parameters")
+
+[/Script/Engine.MaterialExpressionTextureSampleParameterCubeArray]
++MenuCategories=NSLOCTEXT("MaterialExpression", "Texture", "Texture")
++MenuCategories=NSLOCTEXT("MaterialExpression", "Parameters", "Parameters")
+
+[/Script/Engine.MaterialExpressionTextureSampleParameterSubUV]
++MenuCategories=NSLOCTEXT("MaterialExpression", "Texture", "Texture")
++MenuCategories=NSLOCTEXT("MaterialExpression", "Parameters", "Parameters")
++MenuCategories=NSLOCTEXT("MaterialExpression", "Particles", "Particles")
+
+[/Script/Engine.MaterialExpressionTextureSampleParameterVolume]
++MenuCategories=NSLOCTEXT("MaterialExpression", "Texture", "Texture")
++MenuCategories=NSLOCTEXT("MaterialExpression", "Parameters", "Parameters")
+
+[/Script/Engine.MaterialExpressionTime]
++MenuCategories=NSLOCTEXT("MaterialExpression", "Constants", "Constants")
+bShaderInputData=true
+
+[/Script/Engine.MaterialExpressionTransform]
++MenuCategories=NSLOCTEXT("MaterialExpression", "VectorOps", "VectorOps")
+
+[/Script/Engine.MaterialExpressionTransformPosition]
++MenuCategories=NSLOCTEXT("MaterialExpression", "VectorOps", "VectorOps")
+
+[/Script/Engine.MaterialExpressionTruncate]
++MenuCategories=NSLOCTEXT("MaterialExpression", "Math", "Math")
+
+[/Script/Engine.MaterialExpressionTruncateLWC]
++MenuCategories=NSLOCTEXT("MaterialExpression", "Coordinates", "Coordinates")
+
+[/Script/Engine.MaterialExpressionTwoSidedSign]
++MenuCategories=NSLOCTEXT("MaterialExpression", "Constants", "Constants")
+bShaderInputData=true
+
+[/Script/Engine.MaterialExpressionUserSceneTexture]
++MenuCategories=NSLOCTEXT("MaterialExpression", "Texture", "Texture")
+bShaderInputData=true
+
+[/Script/Engine.MaterialExpressionVectorNoise]
++MenuCategories=NSLOCTEXT("MaterialExpressionNoise", "Utility", "Utility")
+
+[/Script/Engine.MaterialExpressionVertexColor]
++MenuCategories=NSLOCTEXT("MaterialExpression", "Constants", "Constants")
+bShaderInputData=true
+
+[/Script/Engine.MaterialExpressionVertexInterpolator]
++MenuCategories=NSLOCTEXT("MaterialExpression", "Utility", "Utility")
+
+[/Script/Engine.MaterialExpressionVertexNormalWS]
++MenuCategories=NSLOCTEXT("MaterialExpression", "Vectors", "Vectors")
++MenuCategories=NSLOCTEXT("MaterialExpression", "Coordinates", "Coordinates")
+bShaderInputData=true
+
+[/Script/Engine.MaterialExpressionVertexTangentWS]
++MenuCategories=NSLOCTEXT("MaterialExpression", "Vectors", "Vectors")
++MenuCategories=NSLOCTEXT("MaterialExpression", "Coordinates", "Coordinates")
+bShaderInputData=true
+
+[/Script/Engine.MaterialExpressionViewProperty]
++MenuCategories=NSLOCTEXT("MaterialExpression", "Constants", "Constants")
+bShaderInputData=true
+
+[/Script/Engine.MaterialExpressionViewSize]
++MenuCategories=NSLOCTEXT("MaterialExpression", "Coordinates", "Coordinates")
+bShaderInputData=true
+
+[/Script/Engine.MaterialExpressionVirtualTextureFeatureSwitch]
++MenuCategories=NSLOCTEXT("MaterialExpression", "VirtualTexture", "VirtualTexture")
+
+[/Script/Engine.MaterialExpressionVolumetricAdvancedMaterialInput]
++MenuCategories=NSLOCTEXT("MaterialExpression", "Volume", "Volume")
+
+[/Script/Engine.MaterialExpressionVolumetricCloudEmptySpaceSkippingInput]
++MenuCategories=NSLOCTEXT("MaterialExpression", "Volume", "Volume")
+
+[/Script/Engine.MaterialExpressionVolumetricAdvancedMaterialOutput]
++MenuCategories=NSLOCTEXT("MaterialExpression", "Volume", "Volume")
+
+[/Script/Engine.MaterialExpressionVolumetricCloudEmptySpaceSkippingOutput]
++MenuCategories=NSLOCTEXT("MaterialExpression", "Volume", "Volume")
+
+[/Script/Engine.MaterialExpressionWorldPosition]
++MenuCategories=NSLOCTEXT("MaterialExpression", "Coordinates", "Coordinates")
+bShaderInputData=true
+
+[/Script/Paper2D.MaterialExpressionSpriteTextureSampler]
++MenuCategories=NSLOCTEXT("MaterialExpression", "Texture", "Texture")
++MenuCategories=NSLOCTEXT("MaterialExpression", "Parameters", "Parameters")
+
+[/Script/Engine.MaterialExternalCodeCollection]
++ExternalCodeDeclarations=(ReturnType=Float3, Name="CameraVector", Definition="Parameters.CameraVector", ShaderFrequency=Pixel|Compute, bIsInlined=true)
++ExternalCodeDeclarations=(ReturnType=Float, Name="EyeAdaptation", Definition="EyeAdaptationLookup()", Derivative=Zero, ShaderFrequency=Pixel, bIsInlined=true)
++ExternalCodeDeclarations=(ReturnType=Float4, Name="FontSignedDistanceData", Definition="GetFontSignedDistanceData(Parameters)", ShaderFrequency=Pixel)
++ExternalCodeDeclarations=(ReturnType=Float, Name="IsOrthographic", Definition="IsOrthoProjectionFloat()", Derivative=Zero, bIsInlined=true)
++ExternalCodeDeclarations=(ReturnType=Float3, Name="LightVector", Definition="Parameters.LightVector", ShaderFrequency=Pixel|Compute, bIsInlined=true)
++ExternalCodeDeclarations=(ReturnType=Float, Name="MeshPaintTextureCoordinateIndex", Definition="GetMeshPaintTextureCoordinateIndex(GetPrimitiveData(Parameters))", bIsInlined=true)
++ExternalCodeDeclarations=(ReturnType=TextureMeshPaint, Name="MeshPaintTextureDescriptor", Definition="GetMeshPaintTextureDescriptor(GetPrimitiveData(Parameters))", Derivative=Zero, bIsInlined=true)
++ExternalCodeDeclarations=(ReturnType=Float4, Name="ParticleColor", Definition="Parameters.Particle.Color", Derivative=Zero, bIsInlined=true)
++ExternalCodeDeclarations=(ReturnType=Float3, Name="ParticleDirection", Definition="Parameters.Particle.Velocity.xyz", Derivative=Zero, bIsInlined=true)
++ExternalCodeDeclarations=(ReturnType=Float2, Name="ParticleMacroUV", Definition="GetParticleMacroUV(Parameters)", ShaderFrequency=Pixel|Compute)
++ExternalCodeDeclarations=(ReturnType=Float, Name="ParticleMotionBlurFade", Definition="Parameters.Particle.MotionBlurFade", Derivative=Zero, bIsInlined=true)
++ExternalCodeDeclarations=(ReturnType=Float, Name="ParticleRadius", Definition="max(Parameters.Particle.TranslatedWorldPositionAndSize.w, .001f)", Derivative=Zero, bIsInlined=true)
++ExternalCodeDeclarations=(ReturnType=Float, Name="ParticleRelativeTime", Definition="Parameters.Particle.RelativeTime", Derivative=Zero, bIsInlined=true)
++ExternalCodeDeclarations=(ReturnType=Float, Name="ParticleRandom", Definition="Parameters.Particle.Random", Derivative=Zero, bIsInlined=true)
++ExternalCodeDeclarations=(ReturnType=Float, Name="ParticleSpeed", Definition="Parameters.Particle.Velocity.w", Derivative=Zero, bIsInlined=true)
++ExternalCodeDeclarations=(ReturnType=Float2, Name="ParticleSize", Definition="Parameters.Particle.Size", Derivative=Zero, bIsInlined=true)
++ExternalCodeDeclarations=(RetunrType=Float2, Name="ParticleSpriteRotation", Definition="float2(Parameters.Particle.SpriteRotation, Parameters.Particle.SpriteRotation * 57.2957795131f)", Derivative=Zero, bIsInlined=true)
++ExternalCodeDeclarations=(ReturnType=Float3, Name="PreSkinnedPosition", Definition="Parameters.PreSkinnedPosition", ShaderFrequency=Vertex, bIsInlined=true)
++ExternalCodeDeclarations=(ReturnType=Float3, Name="PreSkinnedNormal", Definition="Parameters.PreSkinnedNormal", ShaderFrequency=Vertex, bIsInlined=true)
++ExternalCodeDeclarations=(ReturnType=Float3, Name="ReflectionVector", Definition="Parameters.ReflectionVector", ShaderFrequency=Pixel|Compute, bIsInlined=true)
+
diff --git a/Engine/Source/Runtime/Engine/Private/Materials/HLSLMaterialDerivativeAutogen.h b/Engine/Source/Runtime/Engine/Private/Materials/HLSLMaterialDerivativeAutogen.h
--- a/Engine/Source/Runtime/Engine/Private/Materials/HLSLMaterialDerivativeAutogen.h
+++ b/Engine/Source/Runtime/Engine/Private/Materials/HLSLMaterialDerivativeAutogen.h
@@ -5,7 +5,7 @@
 #include "CoreMinimal.h"
 #include "MaterialShared.h"
 
-#if WITH_EDITORONLY_DATA
+#include "HLSLMaterialDerivativeAutogen.generated.h"
 
 class FHLSLMaterialTranslator;
 
@@ -16,7 +16,8 @@ class FHLSLMaterialTranslator;
 * Zero		- This node is aware of partial derivatives, and knows that it's value is zero, as is the case for uniform parameters.
 * Valid		- This node is aware of partial derivatives, and knows that it has a calculated value.
 **/
-enum class EDerivativeStatus
+UENUM()
+enum class EDerivativeStatus : uint8
 {
 	NotAware,
 	NotValid,
@@ -24,6 +25,8 @@ enum class EDerivativeStatus
 	Valid,
 };
 
+#if WITH_EDITORONLY_DATA
+
 inline bool IsDerivativeValid(EDerivativeStatus Status)
 {
 	return Status == EDerivativeStatus::Valid || Status == EDerivativeStatus::Zero;
diff --git a/Engine/Source/Runtime/Engine/Private/Materials/HLSLMaterialTranslator.cpp b/Engine/Source/Runtime/Engine/Private/Materials/HLSLMaterialTranslator.cpp
--- a/Engine/Source/Runtime/Engine/Private/Materials/HLSLMaterialTranslator.cpp
+++ b/Engine/Source/Runtime/Engine/Private/Materials/HLSLMaterialTranslator.cpp
@@ -36,6 +36,7 @@
 #include "Materials/MaterialExpressionVolumetricAdvancedMaterialOutput.h"
 #include "Materials/MaterialExpressionWorldPosition.h"
 #include "Materials/MaterialExpressionFirstPersonOutput.h"
+#include "Materials/MaterialExternalCodeRegistry.h"
 #include "Materials/MaterialSourceTemplate.h"
 #include "Materials/SubstrateMaterial.h"
 #include "StringTemplate.h"
@@ -3831,6 +3832,25 @@ int32 FHLSLMaterialTranslator::AddCodeChunkWithHash(uint64 BaseHash, EMaterialVa
 	return AddCodeChunkInner(Hash, FormattedCode, Type, EDerivativeStatus::NotAware, false);
 }
 
+int32 FHLSLMaterialTranslator::AddExternalCodeChunk(const FName& InExternalCodeIdentifier)
+{
+	if (const FMaterialExternalCodeDeclaration* ExternalCodeDeclaration = MaterialExternalCodeRegistry::Get().FindExternalCode(InExternalCodeIdentifier))
+	{
+		return AddExternalCodeChunk(*ExternalCodeDeclaration);
+	}
+	return Errorf(TEXT("External code identifier '%s' not found"), *InExternalCodeIdentifier.ToString());
+}
+
+int32 FHLSLMaterialTranslator::AddExternalCodeChunk(const FMaterialExternalCodeDeclaration& InExternalCode)
+{
+	if (((int32)InExternalCode.ShaderFrequency & (1 << (int32)ShaderFrequency)) == 0)
+	{
+		FName ShaderFrequencyName = StaticEnum<EMaterialShaderFrequency>()->GetNameByValue(1ull << (int32)ShaderFrequency);
+		return Errorf(TEXT("External code identifier '%s' is not available in the %s shader."), *InExternalCode.Name.ToString(), *ShaderFrequencyName.ToString());
+	}
+	return AddCodeChunkInner(InExternalCode.GetReturnTypeValue(), InExternalCode.Derivative, InExternalCode.bIsInlined, TEXT("%s"), *InExternalCode.Definition);
+}
+
 int32 FHLSLMaterialTranslator::AddUniformExpressionInner(uint64 Hash, FMaterialUniformExpression* UniformExpression, EMaterialValueType Type, const TCHAR* FormattedCode)
 {
 	check(bAllowCodeChunkGeneration);
@@ -5372,7 +5392,7 @@ int32 FHLSLMaterialTranslator::ViewProperty(EMaterialExposedViewProperty Propert
 
 int32 FHLSLMaterialTranslator::IsOrthographic()
 {
-	return AddInlinedCodeChunkZeroDeriv(MCT_Float, TEXT("IsOrthoProjectionFloat()"));
+	return AddExternalCodeChunk(TEXT("IsOrthographic"));
 }
 
 
@@ -5981,15 +6001,11 @@ int32 FHLSLMaterialTranslator::Abs(int32 X)
 
 int32 FHLSLMaterialTranslator::ReflectionVector()
 {
-	if (ShaderFrequency != SF_Pixel && ShaderFrequency != SF_Compute)
-	{
-		return NonPixelShaderExpressionError();
-	}
 	if (ShaderFrequency != SF_Vertex)
 	{
 		bUsesTransformVector = true;
 	}
-	return AddInlinedCodeChunk(MCT_Float3,TEXT("Parameters.ReflectionVector"));
+	return AddExternalCodeChunk(TEXT("ReflectionVector"));
 }
 
 int32 FHLSLMaterialTranslator::ReflectionAboutCustomWorldNormal(int32 CustomWorldNormal, int32 bNormalizeCustomWorldNormal)
@@ -6016,30 +6032,20 @@ int32 FHLSLMaterialTranslator::ReflectionAboutCustomWorldNormal(int32 CustomWorl
 
 int32 FHLSLMaterialTranslator::CameraVector()
 {
-	if (ShaderFrequency != SF_Pixel && ShaderFrequency != SF_Compute)
-	{
-		return NonPixelShaderExpressionError();
-	}
 	if (ShaderFrequency != SF_Vertex)
 	{
 		bUsesTransformVector = true;
 	}
-	return AddInlinedCodeChunk(MCT_Float3,TEXT("Parameters.CameraVector"));
+	return AddExternalCodeChunk(TEXT("CameraVector"));
 }
 
 int32 FHLSLMaterialTranslator::LightVector()
 {
-	if (ShaderFrequency != SF_Pixel && ShaderFrequency != SF_Compute)
-	{
-		return NonPixelShaderExpressionError();
-	}
-
 	if (!Material->IsLightFunction() && !Material->IsDeferredDecal())
 	{
 		return Errorf(TEXT("LightVector can only be used in LightFunction or DeferredDecal materials"));
 	}
-
-	return AddInlinedCodeChunk(MCT_Float3,TEXT("Parameters.LightVector"));
+	return AddExternalCodeChunk(TEXT("LightVector"));
 }
 
 int32 FHLSLMaterialTranslator::GetViewportUV()
@@ -6082,12 +6088,7 @@ int32 FHLSLMaterialTranslator::GetPixelPosition()
 
 int32 FHLSLMaterialTranslator::ParticleMacroUV()
 {
-	if (ShaderFrequency != SF_Pixel && ShaderFrequency != SF_Compute)
-	{
-		return NonPixelShaderExpressionError();
-	}
-
-	return AddCodeChunk(MCT_Float2,TEXT("GetParticleMacroUV(Parameters)"));
+	return AddExternalCodeChunk(TEXT("ParticleMacroUV"));
 }
 
 int32 FHLSLMaterialTranslator::ParticleSubUV(int32 TextureIndex, EMaterialSamplerType SamplerType, int32 MipValue0Index, int32 MipValue1Index, ETextureMipValueMode MipValueMode, bool bBlend)
@@ -6154,12 +6155,8 @@ int32 FHLSLMaterialTranslator::ParticleSubUVProperty(int32 PropertyIndex)
 
 int32 FHLSLMaterialTranslator::ParticleColor()
 {
-	if (ShaderFrequency != SF_Vertex && ShaderFrequency != SF_Pixel && ShaderFrequency != SF_Compute)
-	{
-		return NonVertexOrPixelShaderExpressionError();
-	}
 	bUsesParticleColor |= (ShaderFrequency != SF_Vertex);
-	return AddInlinedCodeChunkZeroDeriv(MCT_Float4,TEXT("Parameters.Particle.Color"));
+	return AddExternalCodeChunk(TEXT("ParticleColor"));
 }
 
 int32 FHLSLMaterialTranslator::ParticlePosition(EPositionOrigin OriginType)
@@ -6192,12 +6189,8 @@ int32 FHLSLMaterialTranslator::ParticlePosition(EPositionOrigin OriginType)
 
 int32 FHLSLMaterialTranslator::ParticleRadius()
 {
-	if (ShaderFrequency != SF_Vertex && ShaderFrequency != SF_Pixel && ShaderFrequency != SF_Compute)
-	{
-		return NonVertexOrPixelShaderExpressionError();
-	}
 	bNeedsParticlePosition = true;
-	return AddInlinedCodeChunkZeroDeriv(MCT_Float,TEXT("max(Parameters.Particle.TranslatedWorldPositionAndSize.w, .001f)"));
+	return AddExternalCodeChunk(TEXT("ParticleRadius"));
 }
 
 int32 FHLSLMaterialTranslator::SphericalParticleOpacity(int32 Density)
@@ -6221,85 +6214,52 @@ int32 FHLSLMaterialTranslator::SphericalParticleOpacity(int32 Density)
 
 int32 FHLSLMaterialTranslator::ParticleRelativeTime()
 {
-	if (ShaderFrequency != SF_Vertex && ShaderFrequency != SF_Pixel && ShaderFrequency != SF_Compute)
-	{
-		return NonVertexOrPixelShaderExpressionError();
-	}
 	bNeedsParticleTime = true;
-	return AddInlinedCodeChunkZeroDeriv(MCT_Float,TEXT("Parameters.Particle.RelativeTime"));
+	return AddExternalCodeChunk(TEXT("ParticleRelativeTime"));
 }
 
 int32 FHLSLMaterialTranslator::ParticleMotionBlurFade()
 {
-	if (ShaderFrequency != SF_Vertex && ShaderFrequency != SF_Pixel && ShaderFrequency != SF_Compute)
-	{
-		return NonVertexOrPixelShaderExpressionError();
-	}
 	bUsesParticleMotionBlur = true;
-	return AddInlinedCodeChunkZeroDeriv(MCT_Float,TEXT("Parameters.Particle.MotionBlurFade"));
+	return AddExternalCodeChunk(TEXT("ParticleMotionBlurFade"));
 }
 
 int32 FHLSLMaterialTranslator::ParticleRandom()
 {
-	if (ShaderFrequency != SF_Vertex && ShaderFrequency != SF_Pixel && ShaderFrequency != SF_Compute)
-	{
-		return NonVertexOrPixelShaderExpressionError();
-	}
 	bNeedsParticleRandom = true;
-	return AddInlinedCodeChunkZeroDeriv(MCT_Float,TEXT("Parameters.Particle.Random"));
+	return AddExternalCodeChunk(TEXT("ParticleRandom"));
 }
 
 
 int32 FHLSLMaterialTranslator::ParticleDirection()
 {
-	if (ShaderFrequency != SF_Vertex && ShaderFrequency != SF_Pixel && ShaderFrequency != SF_Compute)
-	{
-		return NonVertexOrPixelShaderExpressionError();
-	}
 	bNeedsParticleVelocity = true;
-	return AddInlinedCodeChunkZeroDeriv(MCT_Float3,TEXT("Parameters.Particle.Velocity.xyz"));
+	return AddExternalCodeChunk(TEXT("ParticleDirection"));
 }
 
 int32 FHLSLMaterialTranslator::ParticleSpeed()
 {
-	if (ShaderFrequency != SF_Vertex && ShaderFrequency != SF_Pixel && ShaderFrequency != SF_Compute)
-	{
-		return NonVertexOrPixelShaderExpressionError();
-	}
 	bNeedsParticleVelocity = true;
-	return AddInlinedCodeChunkZeroDeriv(MCT_Float,TEXT("Parameters.Particle.Velocity.w"));
+	return AddExternalCodeChunk(TEXT("ParticleSpeed"));
 }
 
 int32 FHLSLMaterialTranslator::ParticleSize()
 {
-	if (ShaderFrequency != SF_Vertex && ShaderFrequency != SF_Pixel && ShaderFrequency != SF_Compute)
-	{
-		return NonVertexOrPixelShaderExpressionError();
-	}
 	bNeedsParticleSize = true;
-	return AddInlinedCodeChunkZeroDeriv(MCT_Float2,TEXT("Parameters.Particle.Size"));
+	return AddExternalCodeChunk(TEXT("ParticleSize"));
 }
 
 int32 FHLSLMaterialTranslator::ParticleSpriteRotation()
 {
-	if (ShaderFrequency != SF_Vertex && ShaderFrequency != SF_Pixel && ShaderFrequency != SF_Compute)
-	{
-		return NonVertexOrPixelShaderExpressionError();
-	}
 	bNeedsParticleSpriteRotation = true;
-	return AddInlinedCodeChunkZeroDeriv(MCT_Float2, TEXT("float2(Parameters.Particle.SpriteRotation, Parameters.Particle.SpriteRotation * 57.2957795131f)"));
+	return AddExternalCodeChunk(TEXT("ParticleSpriteRotation"));
 }
 
 int32 FHLSLMaterialTranslator::LocalPosition(EPositionIncludedOffsets IncludedOffsets, ELocalPositionOrigin OriginType)
 {
 	if (OriginType == ELocalPositionOrigin::InstancePreSkinning)
 	{
-		if (ShaderFrequency != SF_Vertex)
-		{
-			return Errorf(TEXT("Pre-skinned position is only available in the vertex shader, pass through custom interpolators if needed."));
-		}
-
-		return AddInlinedCodeChunk(MCT_Float3,TEXT("Parameters.PreSkinnedPosition"));
+		return AddExternalCodeChunk(TEXT("PreSkinnedPosition"));
 	}
 
 	// If compiling for the previous frame in the vertex shader
@@ -9042,13 +9002,14 @@ int32 FHLSLMaterialTranslator::StaticTerrainLayerWeight(FName LayerName,int32 De
 	}
 }
 
+int32 FHLSLMaterialTranslator::ExternalCode(const FMaterialExternalCodeDeclaration& InExternalCode)
+{
+	return AddExternalCodeChunk(InExternalCode);
+}
+
 int32 FHLSLMaterialTranslator::FontSignedDistanceData()
 {
-	if (ShaderFrequency != SF_Pixel)
-	{
-		return Errorf(TEXT("Font Signed Distance node is only available in pixel shader input."));
-	}
-	return AddCodeChunk(MCT_Float4, TEXT("GetFontSignedDistanceData(Parameters)"));
+	return AddExternalCodeChunk(TEXT("FontSignedDistanceData"));
 }
 
 int32 FHLSLMaterialTranslator::VertexColor()
@@ -9069,12 +9030,12 @@ int32 FHLSLMaterialTranslator::VertexColor()
 
 int32 FHLSLMaterialTranslator::MeshPaintTextureCoordinateIndex()
 {
-	return AddInlinedCodeChunkZeroDeriv(MCT_Float, TEXT("GetMeshPaintTextureCoordinateIndex(GetPrimitiveData(Parameters))"));
+	return AddExternalCodeChunk(TEXT("MeshPaintTextureCoordinateIndex"));
 }
 
 int32 FHLSLMaterialTranslator::MeshPaintTextureDescriptor()
 {
-	return AddInlinedCodeChunkZeroDeriv(MCT_TextureMeshPaint, TEXT("GetMeshPaintTextureDescriptor(GetPrimitiveData(Parameters))"));
+	return AddExternalCodeChunk(TEXT("MeshPaintTextureDescriptor"));
 }
 
 int32 FHLSLMaterialTranslator::MeshPaintTextureReplace(int32 Invalid, int32 Valid)
@@ -9084,22 +9045,12 @@ int32 FHLSLMaterialTranslator::MeshPaintTextureReplace(int32 Invalid, int32 Vali
 
 int32 FHLSLMaterialTranslator::PreSkinnedPosition()
 {
-	if (ShaderFrequency != SF_Vertex)
-	{
-		return Errorf(TEXT("Pre-skinned position is only available in the vertex shader, pass through custom interpolators if needed."));
-	}
-
-	return AddInlinedCodeChunk(MCT_Float3,TEXT("Parameters.PreSkinnedPosition"));
+	return AddExternalCodeChunk(TEXT("PreSkinnedPosition"));
 }
 
 int32 FHLSLMaterialTranslator::PreSkinnedNormal()
 {
-	if (ShaderFrequency != SF_Vertex)
-	{
-		return Errorf(TEXT("Pre-skinned normal is only available in the vertex shader, pass through custom interpolators if needed."));
-	}
-
-	return AddInlinedCodeChunk(MCT_Float3,TEXT("Parameters.PreSkinnedNormal"));
+	return AddExternalCodeChunk(TEXT("PreSkinnedNormal"));
 }
 
 int32 FHLSLMaterialTranslator::VertexInterpolator(uint32 InterpolatorIndex)
@@ -15867,14 +15818,9 @@ int32 FHLSLMaterialTranslator::SpeedTree(int32 GeometryArg, int32 WindArg, int32
 /**Experimental access to the EyeAdaptation RT for Post Process materials. Can be one frame behind depending on the value of BlendableLocation. */
 int32 FHLSLMaterialTranslator::EyeAdaptation()
 {
-	if( ShaderFrequency != SF_Pixel )
-	{
-		return NonPixelShaderExpressionError();
-	}
-
 	MaterialCompilationOutput.bUsesEyeAdaptation = true;
 
-	return AddInlinedCodeChunkZeroDeriv(MCT_Float, TEXT("EyeAdaptationLookup()"));
+	return AddExternalCodeChunk(TEXT("EyeAdaptation"));
 }
 
 /**Experimental access to the EyeAdaptation RT for applying an inverse. */
diff --git a/Engine/Source/Runtime/Engine/Private/Materials/HLSLMaterialTranslator.h b/Engine/Source/Runtime/Engine/Private/Materials/HLSLMaterialTranslator.h
--- a/Engine/Source/Runtime/Engine/Private/Materials/HLSLMaterialTranslator.h
+++ b/Engine/Source/Runtime/Engine/Private/Materials/HLSLMaterialTranslator.h
@@ -801,6 +801,10 @@ protected:
 		return AddCodeChunkInner(Type, EDerivativeStatus::NotValid, true, Format, Args...);
 	}
 
+	// Constructs a code chunk that comes from an external HLSL function. Its return type and function implementation is data-driven and defined in BaseEditor.ini.
+	int32 AddExternalCodeChunk(const FName& InExternalCodeIdentifier);
+	int32 AddExternalCodeChunk(const FMaterialExternalCodeDeclaration& InExternalCode);
+
 	int32 AddUniformExpressionInner(uint64 Hash, FMaterialUniformExpression* UniformExpression, EMaterialValueType Type, const TCHAR* FormattedCode);
 
 	// AddUniformExpression - Adds an input to the Code array and returns its index.
@@ -1130,6 +1134,8 @@ protected:
 
 	virtual int32 StaticTerrainLayerWeight(FName ParameterName, int32 Default, bool bTextureArray = false) override;
 
+	virtual int32 ExternalCode(const FMaterialExternalCodeDeclaration& InExternalCode) override;
+
 	virtual int32 FontSignedDistanceData() override;
 
 	virtual int32 VertexColor() override;
diff --git a/Engine/Source/Runtime/Engine/Private/Materials/MaterialExpressions.cpp b/Engine/Source/Runtime/Engine/Private/Materials/MaterialExpressions.cpp
--- a/Engine/Source/Runtime/Engine/Private/Materials/MaterialExpressions.cpp
+++ b/Engine/Source/Runtime/Engine/Private/Materials/MaterialExpressions.cpp
@@ -97,6 +97,7 @@
 #include "Materials/MaterialExpressionDivide.h"
 #include "Materials/MaterialExpressionDotProduct.h"
 #include "Materials/MaterialExpressionDynamicParameter.h"
+#include "Materials/MaterialExpressionExternalCodeBase.h"
 #include "Materials/MaterialExpressionCloudLayer.h"
 #include "Materials/MaterialExpressionCustomOutput.h"
 #include "Materials/MaterialExpressionEyeAdaptation.h"
@@ -294,6 +295,7 @@
 #include "Materials/MaterialExpressionSparseVolumeTextureObject.h"
 #include "Materials/MaterialExpressionSparseVolumeTextureSample.h"
 #include "Materials/MaterialExpressionIsFirstPerson.h"
+#include "Materials/MaterialExternalCodeRegistry.h"
 #include "EditorSupportDelegates.h"
 #if WITH_EDITOR
 #include "MaterialGraph/MaterialGraphNode_Comment.h"
@@ -3795,11 +3797,6 @@ void UMaterialExpressionMeshPaintTextureCoordinateIndex::GetExpressionToolTip(TA
 	ConvertToMultilineToolTip(TEXT("Get the Mesh Paint Texture UV coordinate index."), 40, OutToolTip);
 }
 
-int32 UMaterialExpressionMeshPaintTextureCoordinateIndex::Compile(class FMaterialCompiler* Compiler, int32 OutputIndex)
-{
-	return Compiler->MeshPaintTextureCoordinateIndex();
-}
-
 #endif // WITH_EDITOR
 
 UMaterialExpressionMeshPaintTextureObject::UMaterialExpressionMeshPaintTextureObject(const FObjectInitializer& ObjectInitializer)
@@ -3811,6 +3808,19 @@ UMaterialExpressionMeshPaintTextureObject::UMaterialExpressionMeshPaintTextureOb
 
 #if WITH_EDITOR
 
+int32 UMaterialExpressionExternalCodeBase::Compile(class FMaterialCompiler* Compiler, int32 OutputIndex)
+{
+	if (const FMaterialExternalCodeDeclaration* ExternalCode = MaterialExternalCodeRegistry::Get().FindExternalCode(ExternalCodeIdentifier))
+	{
+		return Compiler->ExternalCode(*ExternalCode);
+	}
+	return Compiler->Errorf(TEXT("Missing external code declaration for '%s'"), *ExternalCodeIdentifier.ToString());
+}
+
+#endif
+
+#if WITH_EDITOR
+
 void UMaterialExpressionMeshPaintTextureObject::GetCaption(TArray<FString>& OutCaptions) const
 {
 	OutCaptions.Add(TEXT("Mesh Paint Texture Object"));
@@ -3826,11 +3836,6 @@ uint32 UMaterialExpressionMeshPaintTextureObject::GetOutputType(int32 OutputInde
 	return MCT_Texture2D;
 }
 
-int32 UMaterialExpressionMeshPaintTextureObject::Compile(class FMaterialCompiler* Compiler, int32 OutputIndex)
-{
-	return Compiler->MeshPaintTextureDescriptor();
-}
-
 #endif // WITH_EDITOR
 
 UMaterialExpressionMeshPaintTextureReplace::UMaterialExpressionMeshPaintTextureReplace(const FObjectInitializer& ObjectInitializer)
@@ -9821,8 +9826,7 @@ int32 UMaterialExpressionFontSignedDistance::Compile(class FMaterialCompiler* Co
 	{
 		return Compiler->Errorf(TEXT("Font Signed Distance node is only compatible with UI materials"));
 	}
-
-	return Compiler->FontSignedDistanceData();
+	return Super::Compile(Compiler, OutputIndex);
 }
 
 void UMaterialExpressionFontSignedDistance::GetCaption(TArray<FString>& OutCaptions) const
@@ -10102,11 +10106,6 @@ void UMaterialExpressionParticleSubUVProperties::GetCaption(TArray<FString>& Out
 //	MaterialExpressionParticleMacroUV
 //
 #if WITH_EDITOR
-int32 UMaterialExpressionParticleMacroUV::Compile(class FMaterialCompiler* Compiler, int32 OutputIndex)
-{
-	return Compiler->ParticleMacroUV();
-}
-
 void UMaterialExpressionParticleMacroUV::GetCaption(TArray<FString>& OutCaptions) const
 {
 	OutCaptions.Add(TEXT("Particle MacroUV"));
@@ -10207,11 +10206,6 @@ void UMaterialExpressionViewSize::GetCaption(TArray<FString>& OutCaptions) const
 #endif // WITH_EDITOR
 
 #if WITH_EDITOR
-int32 UMaterialExpressionIsOrthographic::Compile(class FMaterialCompiler* Compiler, int32 OutputIndex)
-{
-	return Compiler->IsOrthographic();
-}
-
 void UMaterialExpressionIsOrthographic::GetCaption(TArray<FString>& OutCaptions) const
 {
 	OutCaptions.Add(TEXT("IsOrthographic"));
@@ -21564,7 +21558,7 @@ int32 UMaterialExpressionPreSkinnedPosition::Compile(class FMaterialCompiler* Co
 		return Compiler->Errorf(TEXT("Pre-skinned position is only available in the vertex shader, pass through custom interpolators if needed."));
 	}
 
-	return Compiler->PreSkinnedPosition();
+	return Super::Compile(Compiler, OutputIndex);
 }
 
 void UMaterialExpressionPreSkinnedPosition::GetCaption(TArray<FString>& OutCaptions) const
@@ -21604,11 +21598,6 @@ UMaterialExpressionPreSkinnedNormal::UMaterialExpressionPreSkinnedNormal(const F
 }
 
 #if WITH_EDITOR
-int32 UMaterialExpressionPreSkinnedNormal::Compile(class FMaterialCompiler* Compiler, int32 OutputIndex)
-{
-	return Compiler->PreSkinnedNormal();
-}
-
 void UMaterialExpressionPreSkinnedNormal::GetCaption(TArray<FString>& OutCaptions) const
 {
 	OutCaptions.Add(TEXT("Pre-Skinned Local Normal"));
diff --git a/Engine/Source/Runtime/Engine/Private/Materials/MaterialExternalCodeRegistry.cpp b/Engine/Source/Runtime/Engine/Private/Materials/MaterialExternalCodeRegistry.cpp
new file mode 100644
--- /dev/null
+++ b/Engine/Source/Runtime/Engine/Private/Materials/MaterialExternalCodeRegistry.cpp
@@ -0,0 +1,42 @@
+// Copyright Epic Games, Inc. All Rights Reserved.
+
+#include "MaterialExternalCodeRegistry.h"
+#include "Misc/Parse.h"
+
+MaterialExternalCodeRegistry& MaterialExternalCodeRegistry::Get()
+{
+	static MaterialExternalCodeRegistry Instance;
+	return Instance;
+}
+
+MaterialExternalCodeRegistry::MaterialExternalCodeRegistry()
+{
+	BuildMapToExternalDeclarations();
+}
+
+void MaterialExternalCodeRegistry::BuildMapToExternalDeclarations()
+{
+	if (const UMaterialExternalCodeCollection* ExternalCodeCollection = GetDefault<UMaterialExternalCodeCollection>())
+	{
+		ExternalCodeDeclarationMap.Reserve(ExternalCodeCollection->ExternalCodeDeclarations.Num());
+
+		for (const FMaterialExternalCodeDeclaration& InExternalCode : ExternalCodeCollection->ExternalCodeDeclarations)
+		{
+			if (ExternalCodeDeclarationMap.Contains(InExternalCode.Name))
+			{
+				UE_LOG(LogMaterial, Error, TEXT("External HLSL code declaration in material IR must not be overloaded, but '%s' is defined more than once"), *InExternalCode.Name.ToString());
+				continue;
+			}
+			ExternalCodeDeclarationMap.Add(InExternalCode.Name) = &InExternalCode;
+		}
+	}
+}
+
+const FMaterialExternalCodeDeclaration* MaterialExternalCodeRegistry::FindExternalCode(const FName& InExternalCodeIdentifier) const
+{
+	if (const FMaterialExternalCodeDeclaration* const* Entry = ExternalCodeDeclarationMap.Find(InExternalCodeIdentifier))
+	{
+		return *Entry;
+	}
+	return nullptr;
+}
diff --git a/Engine/Source/Runtime/Engine/Private/Materials/MaterialExternalCodeRegistry.h b/Engine/Source/Runtime/Engine/Private/Materials/MaterialExternalCodeRegistry.h
new file mode 100644
--- /dev/null
+++ b/Engine/Source/Runtime/Engine/Private/Materials/MaterialExternalCodeRegistry.h
@@ -0,0 +1,138 @@
+// Copyright Epic Games, Inc. All Rights Reserved.
+
+/*=============================================================================
+	MaterialExternalCodeRegistry.h: External HLSL code declarations.
+=============================================================================*/
+
+#pragma once
+
+#include "CoreMinimal.h"
+#include "Materials/HLSLMaterialDerivativeAutogen.h"
+#include "RHIDefinitions.h"
+
+#include "MaterialExternalCodeRegistry.generated.h"
+
+UENUM(Flags)
+enum class EMaterialShaderFrequency : uint8
+{
+	Vertex  = 1 << SF_Vertex,
+	Pixel   = 1 << SF_Pixel,
+	Compute = 1 << SF_Compute,
+	Any     = Vertex | Pixel | Compute,
+};
+ENUM_CLASS_FLAGS(EMaterialShaderFrequency);
+
+/**
+ * This has to be kept in sync with EMaterialValueType.
+ * To consolidate those enums, EMaterialValueType must be made an 'enum class' but it has to be deprecated first.
+ * NOTE: Remove UMETA(Hidden) markers once the respective entry is used in BaseMaterialExpressions.ini file.
+ */
+UENUM(Flags)
+enum class EMaterialValueTypeBridge : uint32
+{
+	Float1              = MCT_Float1,
+	Float2              = MCT_Float2,
+	Float3              = MCT_Float3,
+	Float4              = MCT_Float4,
+	Texture2D           = MCT_Texture2D UMETA(Hidden),
+	TextureCube         = MCT_TextureCube UMETA(Hidden),
+	Texture2DArray      = MCT_Texture2DArray UMETA(Hidden),
+	TextureCubeArray    = MCT_TextureCubeArray UMETA(Hidden),
+	VolumeTexture       = MCT_VolumeTexture UMETA(Hidden),
+	StaticBool          = MCT_StaticBool UMETA(Hidden),
+	Unknown             = MCT_Unknown UMETA(Hidden),
+	MaterialAttributes  = MCT_MaterialAttributes UMETA(Hidden),
+	TextureExternal     = MCT_TextureExternal UMETA(Hidden),
+	TextureVirtual      = MCT_TextureVirtual UMETA(Hidden),
+	SparseVolumeTexture = MCT_SparseVolumeTexture UMETA(Hidden),
+	VTPageTableResult   = MCT_VTPageTableResult UMETA(Hidden),
+	ShadingModel        = MCT_ShadingModel UMETA(Hidden),
+	Substrate           = MCT_Substrate UMETA(Hidden),
+	LWCScalar           = MCT_LWCScalar UMETA(Hidden),
+	LWCVector2          = MCT_LWCVector2 UMETA(Hidden),
+	LWCVector3          = MCT_LWCVector3 UMETA(Hidden),
+	LWCVector4          = MCT_LWCVector4 UMETA(Hidden),
+	Execution           = MCT_Execution UMETA(Hidden),
+	VoidStatement       = MCT_VoidStatement UMETA(Hidden),
+	Bool                = MCT_Bool UMETA(Hidden),
+	UInt1               = MCT_UInt1 UMETA(Hidden),
+	UInt2               = MCT_UInt2 UMETA(Hidden),
+	UInt3               = MCT_UInt3 UMETA(Hidden),
+	UInt4               = MCT_UInt4 UMETA(Hidden),
+	TextureCollection   = MCT_TextureCollection UMETA(Hidden),
+	TextureMeshPaint    = MCT_TextureMeshPaint,
+	Texture             = MCT_Texture UMETA(Hidden),
+	Float               = MCT_Float,
+	UInt                = MCT_UInt UMETA(Hidden),
+	LWCType             = MCT_LWCType UMETA(Hidden),
+	Numeric             = MCT_Numeric UMETA(Hidden),
+};
+
+/** Declaration of external HLSL code. Such code expressions can be emitted as part of a material translation. */
+USTRUCT()
+struct FMaterialExternalCodeDeclaration
+{
+	GENERATED_BODY()
+
+public:
+	UPROPERTY(Config)
+	uint32 bIsInlined:1 = false;
+
+	UPROPERTY(Config)
+	EMaterialValueTypeBridge ReturnType = EMaterialValueTypeBridge::Unknown;
+
+	UPROPERTY(Config)
+	FName Name;
+
+	UPROPERTY(Config)
+	FString Definition;
+
+	UPROPERTY(Config)
+	EDerivativeStatus Derivative = EDerivativeStatus::NotAware;
+
+	UPROPERTY(Config)
+	EMaterialShaderFrequency ShaderFrequency = EMaterialShaderFrequency::Any;
+
+public:
+	/**
+	 * Function return type as material value type.
+	 * Deduced from ReturnType string to maintain compatibility with other systems that rely on EMaterialValueType *not* being a UENUM().
+	 */
+	inline EMaterialValueType GetReturnTypeValue() const
+	{
+		return (EMaterialValueType)ReturnType;
+	}
+};
+
+UCLASS(Config=MaterialExpressions)
+class UMaterialExternalCodeCollection : public UObject
+{
+	GENERATED_BODY()
+
+public:
+	UPROPERTY(Config)
+	TArray<FMaterialExternalCodeDeclaration> ExternalCodeDeclarations;
+
+};
+
+/** Singleton class to register external HLSL function and input declarations for material IR modules. */
+class MaterialExternalCodeRegistry
+{
+public:
+	MaterialExternalCodeRegistry(const MaterialExternalCodeRegistry&) = delete;
+	MaterialExternalCodeRegistry& operator = (const MaterialExternalCodeRegistry&) = delete;
+
+	static MaterialExternalCodeRegistry& Get();
+
+	/** Returns the external function declaration for the function with the specified name. */
+	const FMaterialExternalCodeDeclaration* FindExternalCode(const FName& InExternalCodeIdentifier) const;
+
+private:
+	MaterialExternalCodeRegistry();
+
+	/** Builds the name-to-declaration map for all serialized function declarations. */
+	void BuildMapToExternalDeclarations();
+
+	TMap<FName, const FMaterialExternalCodeDeclaration*> ExternalCodeDeclarationMap;
+
+};
diff --git a/Engine/Source/Runtime/Engine/Public/MaterialCompiler.h b/Engine/Source/Runtime/Engine/Public/MaterialCompiler.h
--- a/Engine/Source/Runtime/Engine/Public/MaterialCompiler.h
+++ b/Engine/Source/Runtime/Engine/Public/MaterialCompiler.h
@@ -33,6 +33,7 @@ class UTexture;
 class UTextureCollection;
 struct FMaterialParameterInfo;
 class USparseVolumeTexture;
+struct FMaterialExternalCodeDeclaration;
 
 enum EMaterialForceCastFlags
 {
@@ -482,6 +483,8 @@ public:
 	// Feedback the instanced material layers data after the UMaterialExpressionMaterialAttributeLayers is assembled
 	virtual void FeedbackMaterialLayersInstancedGraphFromCompilation(const FMaterialLayersFunctions* InLayers) {} // TODO: SHould be pure virtual
 
+	/** Adds code to the material that is defined by an external function definition. */
+	virtual int32 ExternalCode(const FMaterialExternalCodeDeclaration& InExternalCode) = 0;
 
 	virtual int32 VertexColor() = 0;
 	virtual int32 MeshPaintTextureCoordinateIndex() = 0;
@@ -1070,6 +1073,8 @@ public:
 	virtual bool GetStaticBoolValue(int32 BoolIndex, bool& bSucceeded) override { return Compiler->GetStaticBoolValue(BoolIndex, bSucceeded); }
 	virtual int32 StaticTerrainLayerWeight(FName ParameterName, int32 Default, bool bTextureArray = false) override { return Compiler->StaticTerrainLayerWeight(ParameterName, Default, bTextureArray); }
 
+	virtual int32 ExternalCode(const FMaterialExternalCodeDeclaration& InExternalCode) override { return Compiler->ExternalCode(InExternalCode); }
+
 	virtual int32 FontSignedDistanceData() override { return Compiler->FontSignedDistanceData(); }
 
 	virtual int32 VertexColor() override { return Compiler->VertexColor(); }
diff --git a/Engine/Source/Runtime/Engine/Public/Materials/MaterialExpression.h b/Engine/Source/Runtime/Engine/Public/Materials/MaterialExpression.h
--- a/Engine/Source/Runtime/Engine/Public/Materials/MaterialExpression.h
+++ b/Engine/Source/Runtime/Engine/Public/Materials/MaterialExpression.h
@@ -139,7 +139,7 @@ struct FMaterialExpressionCollection
 	TObjectPtr<UMaterialExpression> ExpressionExecEnd = nullptr;
 };
 
-UCLASS(abstract, Optional, BlueprintType, Config = "Editor", hidecategories=Object, MinimalAPI)
+UCLASS(abstract, Optional, BlueprintType, Config=MaterialExpressions, hidecategories=Object, MinimalAPI)
 class UMaterialExpression : public UObject
 {
 	GENERATED_UCLASS_BODY()
diff --git a/Engine/Source/Runtime/Engine/Public/Materials/MaterialExpressionExternalCodeBase.h b/Engine/Source/Runtime/Engine/Public/Materials/MaterialExpressionExternalCodeBase.h
new file mode 100644
--- /dev/null
+++ b/Engine/Source/Runtime/Engine/Public/Materials/MaterialExpressionExternalCodeBase.h
@@ -0,0 +1,25 @@
+// Copyright Epic Games, Inc. All Rights Reserved.
+
+#pragma once
+
+#include "Materials/MaterialExpression.h"
+
+#include "MaterialExpressionExternalCodeBase.generated.h"
+
+UCLASS(abstract, hidecategories=Object, MinimalAPI)
+class UMaterialExpressionExternalCodeBase : public UMaterialExpression
+{
+	GENERATED_BODY()
+
+public:
+	/** Identifier of the external HLSL code this expression inserts into the material shader. */
+	UPROPERTY(Config)
+	FName ExternalCodeIdentifier;
+
+	//~ Begin UMaterialExpression Interface
+#if WITH_EDITOR
+	virtual int32 Compile(class FMaterialCompiler* Compiler, int32 OutputIndex) override;
+#endif
+	//~ End UMaterialExpression Interface
+
+};
diff --git a/Engine/Source/Runtime/Engine/Public/Materials/MaterialExpressionFontSignedDistance.h b/Engine/Source/Runtime/Engine/Public/Materials/MaterialExpressionFontSignedDistance.h
--- a/Engine/Source/Runtime/Engine/Public/Materials/MaterialExpressionFontSignedDistance.h
+++ b/Engine/Source/Runtime/Engine/Public/Materials/MaterialExpressionFontSignedDistance.h
@@ -3,11 +3,11 @@
 
 #include "CoreMinimal.h"
 #include "UObject/ObjectMacros.h"
-#include "Materials/MaterialExpression.h"
+#include "Materials/MaterialExpressionExternalCodeBase.h"
 #include "MaterialExpressionFontSignedDistance.generated.h"
 
 UCLASS(collapsecategories, hidecategories=Object, MinimalAPI)
-class UMaterialExpressionFontSignedDistance : public UMaterialExpression
+class UMaterialExpressionFontSignedDistance : public UMaterialExpressionExternalCodeBase
 {
 	GENERATED_UCLASS_BODY()
 
diff --git a/Engine/Source/Runtime/Engine/Public/Materials/MaterialExpressionIsOrthographic.h b/Engine/Source/Runtime/Engine/Public/Materials/MaterialExpressionIsOrthographic.h
--- a/Engine/Source/Runtime/Engine/Public/Materials/MaterialExpressionIsOrthographic.h
+++ b/Engine/Source/Runtime/Engine/Public/Materials/MaterialExpressionIsOrthographic.h
@@ -5,11 +5,11 @@
 
 #include "CoreMinimal.h"
 #include "UObject/ObjectMacros.h"
-#include "Materials/MaterialExpression.h"
+#include "Materials/MaterialExpressionExternalCodeBase.h"
 #include "MaterialExpressionIsOrthographic.generated.h"
 
 UCLASS(collapsecategories, hidecategories=Object)
-class UMaterialExpressionIsOrthographic : public UMaterialExpression
+class UMaterialExpressionIsOrthographic : public UMaterialExpressionExternalCodeBase
 {
 	GENERATED_BODY()
 
@@ -17,7 +17,6 @@ public:
 
 	//~ Begin UMaterialExpression Interface
 #if WITH_EDITOR
-	virtual int32 Compile(class FMaterialCompiler* Compiler, int32 OutputIndex) override;
 	virtual void GetCaption(TArray<FString>& OutCaptions) const override;
 #endif
 	//~ End UMaterialExpression Interface
diff --git a/Engine/Source/Runtime/Engine/Public/Materials/MaterialExpressionMeshPaintTextureCoordinateIndex.h b/Engine/Source/Runtime/Engine/Public/Materials/MaterialExpressionMeshPaintTextureCoordinateIndex.h
--- a/Engine/Source/Runtime/Engine/Public/Materials/MaterialExpressionMeshPaintTextureCoordinateIndex.h
+++ b/Engine/Source/Runtime/Engine/Public/Materials/MaterialExpressionMeshPaintTextureCoordinateIndex.h
@@ -2,12 +2,12 @@
 
 #pragma once
 
-#include "Materials/MaterialExpression.h"
+#include "Materials/MaterialExpressionExternalCodeBase.h"
 
 #include "MaterialExpressionMeshPaintTextureCoordinateIndex.generated.h"
 
 UCLASS(collapsecategories, hidecategories=Object, MinimalAPI)
-class UMaterialExpressionMeshPaintTextureCoordinateIndex : public UMaterialExpression
+class UMaterialExpressionMeshPaintTextureCoordinateIndex : public UMaterialExpressionExternalCodeBase
 {
 	GENERATED_UCLASS_BODY()
 
@@ -15,7 +15,6 @@ class UMaterialExpressionMeshPaintTextureCoordinateIndex : public UMaterialExpre
 #if WITH_EDITOR
 	virtual void GetCaption(TArray<FString>& OutCaptions) const override;
 	virtual void GetExpressionToolTip(TArray<FString>& OutToolTip) override;
-	virtual int32 Compile(class FMaterialCompiler* Compiler, int32 OutputIndex) override;
 #endif
 	//~ End UMaterialExpression Interface
 };
diff --git a/Engine/Source/Runtime/Engine/Public/Materials/MaterialExpressionMeshPaintTextureObject.h b/Engine/Source/Runtime/Engine/Public/Materials/MaterialExpressionMeshPaintTextureObject.h
--- a/Engine/Source/Runtime/Engine/Public/Materials/MaterialExpressionMeshPaintTextureObject.h
+++ b/Engine/Source/Runtime/Engine/Public/Materials/MaterialExpressionMeshPaintTextureObject.h
@@ -2,12 +2,12 @@
 
 #pragma once
 
-#include "Materials/MaterialExpression.h"
+#include "Materials/MaterialExpressionExternalCodeBase.h"
 
 #include "MaterialExpressionMeshPaintTextureObject.generated.h"
 
 UCLASS(collapsecategories, hidecategories=Object, MinimalAPI)
-class UMaterialExpressionMeshPaintTextureObject : public UMaterialExpression
+class UMaterialExpressionMeshPaintTextureObject : public UMaterialExpressionExternalCodeBase
 {
 	GENERATED_UCLASS_BODY()
 
@@ -16,7 +16,6 @@ class UMaterialExpressionMeshPaintTextureObject : public UMaterialExpression
 	virtual void GetCaption(TArray<FString>& OutCaptions) const override;
 	virtual void GetExpressionToolTip(TArray<FString>& OutToolTip) override;
 	virtual uint32 GetOutputType(int32 OutputIndex) override;
-	virtual int32 Compile(class FMaterialCompiler* Compiler, int32 OutputIndex) override;
 #endif
 	//~ End UMaterialExpression Interface
 };
diff --git a/Engine/Source/Runtime/Engine/Public/Materials/MaterialExpressionParticleMacroUV.h b/Engine/Source/Runtime/Engine/Public/Materials/MaterialExpressionParticleMacroUV.h
--- a/Engine/Source/Runtime/Engine/Public/Materials/MaterialExpressionParticleMacroUV.h
+++ b/Engine/Source/Runtime/Engine/Public/Materials/MaterialExpressionParticleMacroUV.h
@@ -9,11 +9,11 @@
 
 #include "CoreMinimal.h"
 #include "UObject/ObjectMacros.h"
-#include "Materials/MaterialExpression.h"
+#include "Materials/MaterialExpressionExternalCodeBase.h"
 #include "MaterialExpressionParticleMacroUV.generated.h"
 
 UCLASS()
-class UMaterialExpressionParticleMacroUV : public UMaterialExpression
+class UMaterialExpressionParticleMacroUV : public UMaterialExpressionExternalCodeBase
 {
 	GENERATED_BODY()
 
@@ -21,7 +21,6 @@ public:
 
 	//~ Begin UMaterialExpression Interface
 #if WITH_EDITOR
-	virtual int32 Compile(class FMaterialCompiler* Compiler, int32 OutputIndex) override;
 	virtual void GetCaption(TArray<FString>& OutCaptions) const override;
 
 #endif
diff --git a/Engine/Source/Runtime/Engine/Public/Materials/MaterialExpressionPreSkinnedNormal.h b/Engine/Source/Runtime/Engine/Public/Materials/MaterialExpressionPreSkinnedNormal.h
--- a/Engine/Source/Runtime/Engine/Public/Materials/MaterialExpressionPreSkinnedNormal.h
+++ b/Engine/Source/Runtime/Engine/Public/Materials/MaterialExpressionPreSkinnedNormal.h
@@ -4,17 +4,16 @@
 
 #include "CoreMinimal.h"
 #include "UObject/ObjectMacros.h"
-#include "Materials/MaterialExpression.h"
+#include "Materials/MaterialExpressionExternalCodeBase.h"
 #include "MaterialExpressionPreSkinnedNormal.generated.h"
 
 UCLASS(collapsecategories, hidecategories=Object, MinimalAPI)
-class UMaterialExpressionPreSkinnedNormal : public UMaterialExpression
+class UMaterialExpressionPreSkinnedNormal : public UMaterialExpressionExternalCodeBase
 {
 	GENERATED_UCLASS_BODY()
 
 	//~ Begin UMaterialExpression Interface
 #if WITH_EDITOR
-	virtual int32 Compile(class FMaterialCompiler* Compiler, int32 OutputIndex) override;
 	virtual void GetCaption(TArray<FString>& OutCaptions) const override;
 	virtual void GetExpressionToolTip(TArray<FString>& OutToolTip) override;
 	virtual FText GetKeywords() const override {return FText::FromString(TEXT("normal"));}
diff --git a/Engine/Source/Runtime/Engine/Public/Materials/MaterialExpressionPreSkinnedPosition.h b/Engine/Source/Runtime/Engine/Public/Materials/MaterialExpressionPreSkinnedPosition.h
--- a/Engine/Source/Runtime/Engine/Public/Materials/MaterialExpressionPreSkinnedPosition.h
+++ b/Engine/Source/Runtime/Engine/Public/Materials/MaterialExpressionPreSkinnedPosition.h
@@ -4,11 +4,11 @@
 
 #include "CoreMinimal.h"
 #include "UObject/ObjectMacros.h"
-#include "Materials/MaterialExpression.h"
+#include "Materials/MaterialExpressionExternalCodeBase.h"
 #include "MaterialExpressionPreSkinnedPosition.generated.h"
 
 UCLASS(collapsecategories, hidecategories=Object, MinimalAPI)
-class UMaterialExpressionPreSkinnedPosition : public UMaterialExpression
+class UMaterialExpressionPreSkinnedPosition : public UMaterialExpressionExternalCodeBase
 {
 	GENERATED_UCLASS_BODY()
 
