[Backout] - CL39867291
#rnx
Original CL Desc
-----------------------------------------------------------------
[Materials] Initial work on data-driven external code declarations for new material translator.

This is the initial change to move the material translator (including the legacy FHLSLMaterialTranslator) towards a data-driven approach for all external HLSL code injections.
FMaterialExternalCodeDeclaration represents the main structure to organize external HLSL code expressions.
This structure is managed by the singleton MaterialExternalCodeRegistry and referenced by material expressions that extent UMaterialExpressionExternalCodeBase.
In FHLSLMaterialTranslator, they can also be invoked by their name, for material expressions that partially construct their output out of custom HLSL code.
In the upcoming changes, material expression classes that solely inject external HLSL code will be successively replaced by the new UMaterialExpressionExternalCodeBase class.

diff --git a/Engine/Config/BaseEditor.ini b/Engine/Config/BaseEditor.ini
--- a/Engine/Config/BaseEditor.ini
+++ b/Engine/Config/BaseEditor.ini
@@ -708,3 +708,941 @@ bUseWidgetTemplateSelector=True
 +Filters=(FilterString="TexCoordRHIBuffer",DisplayName="Tex Coord Buffers")
 +Filters=(FilterString="FColorVertexBuffer",DisplayName="Color Vertex Buffers")
 
+[/Script/Engine.MaterialExpressionBentNormalCustomOutput]
++MenuCategories=NSLOCTEXT("MaterialExpression", "Utility", "Utility")
+
+[/Script/Engine.MaterialExpressionClearCoatNormalCustomOutput]
++MenuCategories=NSLOCTEXT("MaterialExpression", "Utility", "Utility")
+
+[/Script/Engine.MaterialExpressionTangentOutput]
++MenuCategories=NSLOCTEXT("MaterialExpression", "Custom", "Custom")
+
+[/Script/Engine.MaterialExpressionThinTranslucentMaterialOutput]
++MenuCategories=NSLOCTEXT("MaterialExpression", "ThinTranslucent", "ThinTranslucent")
+
+[/Script/Engine.MaterialExpressionFirstPersonOutput]
++MenuCategories=NSLOCTEXT("MaterialExpression", "FirstPerson", "First Person")
+
+[/Script/Landscape.MaterialExpressionLandscapeGrassOutput]
++MenuCategories=NSLOCTEXT("Landscape", "Landscape", "Landscape")
+
+[/Script/Landscape.MaterialExpressionLandscapeLayerBlend]
++MenuCategories=NSLOCTEXT("Landscape", "Landscape", "Landscape")
+
+[/Script/Landscape.MaterialExpressionLandscapeLayerCoords]
++MenuCategories=NSLOCTEXT("Landscape", "Landscape", "Landscape")
+
+[/Script/Landscape.MaterialExpressionLandscapeLayerSample]
++MenuCategories=NSLOCTEXT("Landscape", "Landscape", "Landscape")
+
+[/Script/Landscape.MaterialExpressionLandscapeLayerSwitch]
++MenuCategories=NSLOCTEXT("Landscape", "Landscape", "Landscape")
+
+[/Script/Landscape.MaterialExpressionLandscapeLayerWeight]
++MenuCategories=NSLOCTEXT("Landscape", "Landscape", "Landscape")
+
+[/Script/Landscape.MaterialExpressionLandscapePhysicalMaterialOutput]
++MenuCategories=NSLOCTEXT("Landscape", "Landscape", "Landscape")
+
+[/Script/Landscape.MaterialExpressionLandscapeVisibilityMask]
++MenuCategories=NSLOCTEXT("Landscape", "Landscape", "Landscape")
+
+[/Script/Engine.MaterialExpressionAbs]
++MenuCategories=NSLOCTEXT("MaterialExpression", "Math", "Math")
+
+[/Script/Engine.MaterialExpressionAbsorptionMediumMaterialOutput]
++MenuCategories=NSLOCTEXT("MaterialExpression", "Shading", "Shading")
+
+[/Script/Engine.MaterialExpressionActorPositionWS]
++MenuCategories=NSLOCTEXT("MaterialExpression", "Vectors", "Vectors")
++MenuCategories=NSLOCTEXT("MaterialExpression", "Coordinates", "Coordinates")
+bShaderInputData=true
+
+[/Script/Engine.MaterialExpressionAdd]
++MenuCategories=NSLOCTEXT("MaterialExpression", "Math", "Math")
+
+[/Script/Engine.MaterialExpressionTextureSample]
++MenuCategories=NSLOCTEXT("MaterialExpression", "Texture", "Texture")
+
+[/Script/Engine.MaterialExpressionTextureSampleParameter2D]
++MenuCategories=NSLOCTEXT("MaterialExpression", "Texture", "Texture")
++MenuCategories=NSLOCTEXT("MaterialExpression", "Parameters", "Parameters")
+
+[/Script/Engine.MaterialExpressionAntialiasedTextureMask]
++MenuCategories=NSLOCTEXT("MaterialExpression", "Utility", "Utility")
+
+[/Script/Engine.MaterialExpressionAppendVector]
++MenuCategories=NSLOCTEXT("MaterialExpression", "Math", "Math")
++MenuCategories=NSLOCTEXT("MaterialExpression", "VectorOps", "VectorOps")
+
+[/Script/Engine.MaterialExpressionArccosine]
++MenuCategories=NSLOCTEXT("MaterialExpression", "Math", "Math")
+
+[/Script/Engine.MaterialExpressionArccosineFast]
++MenuCategories=NSLOCTEXT("MaterialExpression", "Math", "Math")
+
+[/Script/Engine.MaterialExpressionArcsine]
++MenuCategories=NSLOCTEXT("MaterialExpression", "Math", "Math")
+
+[/Script/Engine.MaterialExpressionArcsineFast]
++MenuCategories=NSLOCTEXT("MaterialExpression", "Math", "Math")
+
+[/Script/Engine.MaterialExpressionArctangent]
++MenuCategories=NSLOCTEXT("MaterialExpression", "Math", "Math")
+
+[/Script/Engine.MaterialExpressionArctangent2]
++MenuCategories=NSLOCTEXT("MaterialExpression", "Math", "Math")
+
+[/Script/Engine.MaterialExpressionArctangent2Fast]
++MenuCategories=NSLOCTEXT("MaterialExpression", "Math", "Math")
+
+[/Script/Engine.MaterialExpressionArctangentFast]
++MenuCategories=NSLOCTEXT("MaterialExpression", "Math", "Math")
+
+[/Script/Engine.MaterialExpressionAtmosphericFogColor]
++MenuCategories=NSLOCTEXT("MaterialExpression", "Atmosphere", "Atmosphere")
+
+[/Script/Engine.MaterialExpressionAtmosphericLightColor]
++MenuCategories=NSLOCTEXT("MaterialExpression", "Utility", "Utility")
+
+[/Script/Engine.MaterialExpressionAtmosphericLightVector]
++MenuCategories=NSLOCTEXT("MaterialExpression", "Utility", "Utility")
+bShaderInputData=true
+
+[/Script/Engine.MaterialExpressionBindlessSwitch]
++MenuCategories=NSLOCTEXT("MaterialExpression", "Utility", "Utility")
+
+[/Script/Engine.MaterialExpressionBlackBody]
++MenuCategories=NSLOCTEXT("MaterialExpression", "Utility", "Utility")
+
+[/Script/Engine.MaterialExpressionBlendMaterialAttributes]
++MenuCategories=NSLOCTEXT("MaterialExpression", "MaterialAttributes", "Material Attributes")
+
+[/Script/Engine.MaterialExpressionLegacyBlendMaterialAttributes]
++MenuCategories=NSLOCTEXT("MaterialExpression", "MaterialAttributes", "Material Attributes")
+
+[/Script/Engine.MaterialExpressionBounds]
++MenuCategories=NSLOCTEXT("MaterialExpression", "Vectors", "Vectors")
+bShaderInputData=true
+
+[/Script/Engine.MaterialExpressionBreakMaterialAttributes]
++MenuCategories=NSLOCTEXT("MaterialExpression", "MaterialAttributes", "Material Attributes")
+
+[/Script/Engine.MaterialExpressionBumpOffset]
++MenuCategories=NSLOCTEXT("MaterialExpression", "Utility", "Utility")
+
+[/Script/Engine.MaterialExpressionCameraPositionWS]
++MenuCategories=NSLOCTEXT("MaterialExpression", "Vectors", "Vectors")
++MenuCategories=NSLOCTEXT("MaterialExpression", "Coordinates", "Coordinates")
+bShaderInputData=true
+
+[/Script/Engine.MaterialExpressionCameraVectorWS]
++MenuCategories=NSLOCTEXT("MaterialExpression", "Vectors", "Vectors")
+bShaderInputData=true
+
+[/Script/Engine.MaterialExpressionCeil]
++MenuCategories=NSLOCTEXT("MaterialExpression", "Math", "Math")
+
+[/Script/Engine.MaterialExpressionComment]
++MenuCategories=NSLOCTEXT("MaterialExpression", "Utility", "Utility")
+
+[/Script/Engine.MaterialExpressionVectorParameter]
++MenuCategories=NSLOCTEXT("MaterialExpression", "Parameters", "Parameters")
+
+[/Script/Engine.MaterialExpressionChannelMaskParameter]
++MenuCategories=NSLOCTEXT("MaterialExpression", "Parameters", "Parameters")
++MenuCategories=NSLOCTEXT("MaterialExpression", "Utility", "Utility")
+
+[/Script/Engine.MaterialExpressionClamp]
++MenuCategories=NSLOCTEXT("MaterialExpression", "Math", "Math")
+
+[/Script/Engine.MaterialExpressionCloudSampleAttribute]
++MenuCategories=NSLOCTEXT("MaterialExpression", "Volume", "Volume")
+
+[/Script/Engine.MaterialExpressionCollectionParameter]
++MenuCategories=NSLOCTEXT("MaterialExpression", "Parameters", "Parameters")
+
+[/Script/Engine.MaterialExpressionCollectionTransform]
++MenuCategories=NSLOCTEXT("MaterialExpression", "Parameters", "Parameters")
+
+[/Script/Engine.MaterialExpressionComponentMask]
++MenuCategories=NSLOCTEXT("MaterialExpression", "Math", "Math")
++MenuCategories=NSLOCTEXT("MaterialExpression", "VectorOps", "VectorOps")
+
+[/Script/Engine.MaterialExpressionComposite]
++MenuCategories=NSLOCTEXT("MaterialExpression", "Utility", "Utility")
+
+[/Script/Engine.MaterialExpressionConstant]
++MenuCategories=NSLOCTEXT("MaterialExpression", "Constants", "Constants")
+
+[/Script/Engine.MaterialExpressionConstant2Vector]
++MenuCategories=NSLOCTEXT("MaterialExpression", "Constants", "Constants")
++MenuCategories=NSLOCTEXT("MaterialExpression", "Vectors", "Vectors")
+
+[/Script/Engine.MaterialExpressionConstant3Vector]
++MenuCategories=NSLOCTEXT("MaterialExpression", "Constants", "Constants")
++MenuCategories=NSLOCTEXT("MaterialExpression", "Vectors", "Vectors")
+
+[/Script/Engine.MaterialExpressionConstant4Vector]
++MenuCategories=NSLOCTEXT("MaterialExpression", "Constants", "Constants")
++MenuCategories=NSLOCTEXT("MaterialExpression", "Vectors", "Vectors")
+
+[/Script/Engine.MaterialExpressionConstantBiasScale]
++MenuCategories=NSLOCTEXT("MaterialExpression", "Utility", "Utility")
+
+[/Script/Engine.MaterialExpressionCosine]
++MenuCategories=NSLOCTEXT("MaterialExpression", "Math", "Math")
+
+[/Script/Engine.MaterialExpressionCrossProduct]
++MenuCategories=NSLOCTEXT("MaterialExpression", "Math", "Math")
++MenuCategories=NSLOCTEXT("MaterialExpression", "VectorOps", "VectorOps")
+
+[/Script/Engine.MaterialExpressionScalarParameter]
++MenuCategories=NSLOCTEXT("MaterialExpression", "Parameters", "Parameters")
+
+[/Script/Engine.MaterialExpressionCurveAtlasRowParameter]
++MenuCategories=NSLOCTEXT("MaterialExpression", "Parameters", "Parameters")
+
+[/Script/Engine.MaterialExpressionCustom]
++MenuCategories=NSLOCTEXT("MaterialExpression", "Custom", "Custom")
+
+[/Script/Engine.MaterialExpressionDataDrivenShaderPlatformInfoSwitch]
++MenuCategories=NSLOCTEXT("MaterialExpression", "Utility", "Utility")
+
+[/Script/Engine.MaterialExpressionDBufferTexture]
++MenuCategories=NSLOCTEXT("MaterialExpression", "Decals", "Decals")
+bShaderInputData=true
+
+[/Script/Engine.MaterialExpressionDDX]
++MenuCategories=NSLOCTEXT("MaterialExpression", "Utility", "Utility")
+
+[/Script/Engine.MaterialExpressionDDY]
++MenuCategories=NSLOCTEXT("MaterialExpression", "Utility", "Utility")
+
+[/Script/Engine.MaterialExpressionDecalColor]
++MenuCategories=NSLOCTEXT("MaterialExpression", "Utils", "Utils")
+bShaderInputData=true
+
+[/Script/Engine.MaterialExpressionDecalDerivative]
++MenuCategories=NSLOCTEXT("MaterialExpression", "Utils", "Utils")
+bShaderInputData=true
+
+[/Script/Engine.MaterialExpressionDecalLifetimeOpacity]
++MenuCategories=NSLOCTEXT("MaterialExpression", "Utils", "Utils")
+bShaderInputData=true
+
+[/Script/Engine.MaterialExpressionDecalMipmapLevel]
++MenuCategories=NSLOCTEXT("MaterialExpression", "Utils", "Utils")
+
+[/Script/Engine.MaterialExpressionDeltaTime]
++MenuCategories=NSLOCTEXT("MaterialExpression", "Constants", "Constants")
+bShaderInputData=true
+
+[/Script/Engine.MaterialExpressionDepthFade]
++MenuCategories=NSLOCTEXT("MaterialExpression", "Depth", "Depth")
++MenuCategories=NSLOCTEXT("MaterialExpression", "Utility", "Utility")
+
+[/Script/Engine.MaterialExpressionDepthOfFieldFunction]
++MenuCategories=NSLOCTEXT("MaterialExpression", "Utility", "Utility")
+
+[/Script/Engine.MaterialExpressionDeriveNormalZ]
++MenuCategories=NSLOCTEXT("MaterialExpression", "VectorOps", "VectorOps")
+
+[/Script/Engine.MaterialExpressionDesaturation]
++MenuCategories=NSLOCTEXT("MaterialExpression", "Color", "Color")
++MenuCategories=NSLOCTEXT("MaterialExpression", "Utility", "Utility")
+
+[/Script/Engine.MaterialExpressionDistance]
++MenuCategories=NSLOCTEXT("MaterialExpression", "Utility", "Utility")
+
+[/Script/Engine.MaterialExpressionDistanceCullFade]
++MenuCategories=NSLOCTEXT("MaterialExpression", "Constants", "Constants")
+bShaderInputData=true
+
+[/Script/Engine.MaterialExpressionDistanceFieldApproxAO]
++MenuCategories=NSLOCTEXT("MaterialExpression", "Utility", "Utility")
+
+[/Script/Engine.MaterialExpressionDistanceFieldGradient]
++MenuCategories=NSLOCTEXT("MaterialExpression", "Utility", "Utility")
+
+[/Script/Engine.MaterialExpressionDistanceFieldsRenderingSwitch]
++MenuCategories=NSLOCTEXT("MaterialExpression", "DistanceFieldsRendering", "DistanceFieldsRendering")
+
+[/Script/Engine.MaterialExpressionDistanceToNearestSurface]
++MenuCategories=NSLOCTEXT("MaterialExpression", "Utility", "Utility")
+
+[/Script/Engine.MaterialExpressionDivide]
++MenuCategories=NSLOCTEXT("MaterialExpression", "Math", "Math")
+
+[/Script/Engine.MaterialExpressionDotProduct]
++MenuCategories=NSLOCTEXT("MaterialExpression", "Math", "Math")
++MenuCategories=NSLOCTEXT("MaterialExpression", "VectorOps", "VectorOps")
+
+[/Script/Engine.MaterialExpressionDoubleVectorParameter]
++MenuCategories=NSLOCTEXT("MaterialExpression", "Parameters", "Parameters")
+
+[/Script/Engine.MaterialExpressionDynamicParameter]
++MenuCategories=NSLOCTEXT("MaterialExpression", "Particles", "Particles")
++MenuCategories=NSLOCTEXT("MaterialExpression", "Parameters", "Parameters")
+bShaderInputData=true
+
+[/Script/Engine.MaterialExpressionExponential]
++MenuCategories=NSLOCTEXT("MaterialExpression", "Math", "Math")
+
+[/Script/Engine.MaterialExpressionExponential2]
++MenuCategories=NSLOCTEXT("MaterialExpression", "Math", "Math")
+
+[/Script/Engine.MaterialExpressionEyeAdaptation]
++MenuCategories=NSLOCTEXT("MaterialExpression", "Utility", "Utility")
+bShaderInputData=true
+
+[/Script/Engine.MaterialExpressionEyeAdaptationInverse]
++MenuCategories=NSLOCTEXT("MaterialExpression", "Utility", "Utility")
+bShaderInputData=true
+
+[/Script/Engine.MaterialExpressionFeatureLevelSwitch]
++MenuCategories=NSLOCTEXT("MaterialExpression", "Utility", "Utility")
+
+[/Script/Engine.MaterialExpressionFloatToUInt]
++MenuCategories=NSLOCTEXT("MaterialExpression", "Math", "Math")
+
+[/Script/Engine.MaterialExpressionUIntToFloat]
++MenuCategories=NSLOCTEXT("MaterialExpression", "Math", "Math")
+
+[/Script/Engine.MaterialExpressionFloor]
++MenuCategories=NSLOCTEXT("MaterialExpression", "Math", "Math")
+
+[/Script/Engine.MaterialExpressionFmod]
++MenuCategories=NSLOCTEXT("MaterialExpression", "Math", "Math")
+
+[/Script/Engine.MaterialExpressionFontSample]
++MenuCategories=NSLOCTEXT("MaterialExpression", "Font", "Font")
++MenuCategories=NSLOCTEXT("MaterialExpression", "Texture", "Texture")
+
+[/Script/Engine.MaterialExpressionFontSampleParameter]
++MenuCategories=NSLOCTEXT("MaterialExpression", "Font", "Font")
++MenuCategories=NSLOCTEXT("MaterialExpression", "Texture", "Texture")
++MenuCategories=NSLOCTEXT("MaterialExpression", "Font", "Font")
++MenuCategories=NSLOCTEXT("MaterialExpression", "Parameters", "Parameters")
+
+[/Script/Engine.MaterialExpressionFontSignedDistance]
++MenuCategories=NSLOCTEXT("MaterialExpression", "Constants", "Constants")
+
+[/Script/Engine.MaterialExpressionFrac]
++MenuCategories=NSLOCTEXT("MaterialExpression", "Math", "Math")
+
+[/Script/Engine.MaterialExpressionFresnel]
++MenuCategories=NSLOCTEXT("MaterialExpression", "Utility", "Utility")
+
+[/Script/Engine.MaterialExpressionFunctionInput]
++MenuCategories=NSLOCTEXT("MaterialExpression", "Functions", "Functions")
+
+[/Script/Engine.MaterialExpressionFunctionOutput]
++MenuCategories=NSLOCTEXT("MaterialExpression", "Functions", "Functions")
+
+[/Script/Engine.MaterialExpressionConstantDouble]
++MenuCategories=NSLOCTEXT("MaterialExpression", "Constants", "Constants")
+
+[/Script/Engine.MaterialExpressionGetMaterialAttributes]
++MenuCategories=NSLOCTEXT("MaterialExpression", "MaterialAttributes", "Material Attributes")
+
+[/Script/Engine.MaterialExpressionGIReplace]
++MenuCategories=NSLOCTEXT("MaterialExpression", "Utility", "Utility")
+
+[/Script/Engine.MaterialExpressionHairAttributes]
++MenuCategories=NSLOCTEXT("MaterialExpression", "Hair Attributes", "Hair Attributes")
+
+[/Script/Engine.MaterialExpressionHairColor]
++MenuCategories=NSLOCTEXT("MaterialExpression", "Hair Color", "Hair Color")
+
+[/Script/Engine.MaterialExpressionHsvToRgb]
++MenuCategories=NSLOCTEXT("MaterialExpression", "Math", "Math")
+
+[/Script/Engine.MaterialExpressionIf]
++MenuCategories=NSLOCTEXT("MaterialExpression", "Math", "Math")
+
+[/Script/Engine.MaterialExpressionIfThenElse]
++MenuCategories=NSLOCTEXT("MaterialExpression", "Execution", "Execution")
+
+[/Script/Engine.MaterialExpressionInverseLinearInterpolate]
++MenuCategories=NSLOCTEXT("MaterialExpression", "Math", "Math")
++MenuCategories=NSLOCTEXT("MaterialExpression", "Utility", "Utility")
+
+[/Script/Engine.MaterialExpressionIsFirstPerson]
++MenuCategories=NSLOCTEXT("MaterialExpression", "Constants", "Constants")
++MenuCategories=NSLOCTEXT("MaterialExpression", "FirstPerson", "First Person")
+bShaderInputData=true
+
+[/Script/Engine.MaterialExpressionIsOrthographic]
++MenuCategories=NSLOCTEXT("MaterialExpression", "Constants", "Constants")
+bShaderInputData=true
+
+[/Script/Engine.MaterialExpressionLength]
++MenuCategories=NSLOCTEXT("MaterialExpression", "Utility", "Utility")
+
+[/Script/Engine.MaterialExpressionLightmapUVs]
++MenuCategories=NSLOCTEXT("MaterialExpression", "Coordinates", "Coordinates")
+bShaderInputData=true
+
+[/Script/Engine.MaterialExpressionLightmassReplace]
++MenuCategories=NSLOCTEXT("MaterialExpression", "Utility", "Utility")
+
+[/Script/Engine.MaterialExpressionLightVector]
++MenuCategories=NSLOCTEXT("MaterialExpression", "Vectors", "Vectors")
+
+[/Script/Engine.MaterialExpressionLinearInterpolate]
++MenuCategories=NSLOCTEXT("MaterialExpression", "Math", "Math")
++MenuCategories=NSLOCTEXT("MaterialExpression", "Utility", "Utility")
+
+[/Script/Engine.MaterialExpressionLocalPosition]
++MenuCategories=NSLOCTEXT("MaterialExpression", "Coordinates", "Coordinates")
+bShaderInputData=true
+
+[/Script/Engine.MaterialExpressionLogarithm]
++MenuCategories=NSLOCTEXT("MaterialExpression", "Math", "Math")
+
+[/Script/Engine.MaterialExpressionLogarithm10]
++MenuCategories=NSLOCTEXT("MaterialExpression", "Math", "Math")
+
+[/Script/Engine.MaterialExpressionLogarithm2]
++MenuCategories=NSLOCTEXT("MaterialExpression", "Math", "Math")
+
+[/Script/Engine.MaterialExpressionMakeMaterialAttributes]
++MenuCategories=NSLOCTEXT("MaterialExpression", "MaterialAttributes", "Material Attributes")
+
+[/Script/Engine.MaterialExpressionMapARPassthroughCameraUV]
++MenuCategories=NSLOCTEXT("MaterialExpression", "Coordinates", "Coordinates")
+
+[/Script/Engine.MaterialExpressionMaterialAttributeLayers]
++MenuCategories=NSLOCTEXT("MaterialExpression", "MaterialAttributes", "Material Attributes")
++MenuCategories=NSLOCTEXT("MaterialExpression", "Parameters", "Parameters")
+
+[/Script/Engine.MaterialExpressionMaterialFunctionCall]
++MenuCategories=NSLOCTEXT("MaterialExpression", "Functions", "Functions")
+
+[/Script/Engine.MaterialExpressionMaterialProxyReplace]
++MenuCategories=NSLOCTEXT("MaterialExpression", "Utility", "Utility")
+
+[/Script/Engine.MaterialExpressionMax]
++MenuCategories=NSLOCTEXT("MaterialExpression", "Math", "Math")
+
+[/Script/Engine.MaterialExpressionMeshPaintTextureCoordinateIndex]
++MenuCategories=NSLOCTEXT("MaterialExpression", "MeshPaintTexture", "MeshPaintTexture")
+
+[/Script/Engine.MaterialExpressionMeshPaintTextureObject]
++MenuCategories=NSLOCTEXT("MaterialExpression", "MeshPaintTexture", "MeshPaintTexture")
+
+[/Script/Engine.MaterialExpressionMeshPaintTextureReplace]
++MenuCategories=NSLOCTEXT("MaterialExpression", "MeshPaintTexture", "MeshPaintTexture")
+
+[/Script/Engine.MaterialExpressionMin]
++MenuCategories=NSLOCTEXT("MaterialExpression", "Math", "Math")
+
+[/Script/Engine.MaterialExpressionModulo]
++MenuCategories=NSLOCTEXT("MaterialExpression", "Math", "Math")
+
+[/Script/Engine.MaterialExpressionMultiply]
++MenuCategories=NSLOCTEXT("MaterialExpression", "Math", "Math")
+
+[/Script/Engine.MaterialExpressionNamedRerouteDeclaration]
++MenuCategories=NSLOCTEXT("MaterialExpression", "Utility", "Utility")
+
+[/Script/Engine.MaterialExpressionNaniteReplace]
++MenuCategories=NSLOCTEXT("MaterialExpression", "Utility", "Utility")
+
+[/Script/Engine.MaterialExpressionNeuralNetworkInput]
++MenuCategories=NSLOCTEXT("MaterialExpression", "NeuralNetwork", "NeuralNetwork")
+
+[/Script/Engine.MaterialExpressionNeuralNetworkOutput]
++MenuCategories=NSLOCTEXT("MaterialExpression", "NeuralNetwork", "NeuralNetwork")
+bShaderInputData=true
+
+[/Script/Engine.MaterialExpressionNoise]
++MenuCategories=NSLOCTEXT("MaterialExpressionNoise", "Utility", "Utility")
+
+[/Script/Engine.MaterialExpressionScalarBlueNoise]
++MenuCategories=NSLOCTEXT("MaterialExpressionNoise", "Noise", "Noise")
+
+[/Script/Engine.MaterialExpressionNormalize]
++MenuCategories=NSLOCTEXT("MaterialExpression", "Math", "Math")
++MenuCategories=NSLOCTEXT("MaterialExpression", "VectorOps", "VectorOps")
+
+[/Script/Engine.MaterialExpressionObjectBounds]
++MenuCategories=NSLOCTEXT("MaterialExpression", "Vectors", "Vectors")
+bShaderInputData=true
+
+[/Script/Engine.MaterialExpressionObjectLocalBounds]
++MenuCategories=NSLOCTEXT("MaterialExpression", "Vectors", "Vectors")
+bShaderInputData=true
+
+[/Script/Engine.MaterialExpressionObjectOrientation]
++MenuCategories=NSLOCTEXT("MaterialExpression", "Vectors", "Vectors")
++MenuCategories=NSLOCTEXT("MaterialExpression", "Coordinates", "Coordinates")
+bShaderInputData=true
+
+[/Script/Engine.MaterialExpressionObjectPositionWS]
++MenuCategories=NSLOCTEXT("MaterialExpression", "Vectors", "Vectors")
++MenuCategories=NSLOCTEXT("MaterialExpression", "Coordinates", "Coordinates")
+bShaderInputData=true
+
+[/Script/Engine.MaterialExpressionObjectRadius]
++MenuCategories=NSLOCTEXT("MaterialExpression", "Coordinates", "Coordinates")
+bShaderInputData=true
+
+[/Script/Engine.MaterialExpressionOneMinus]
++MenuCategories=NSLOCTEXT("MaterialExpression", "Math", "Math")
+
+[/Script/Engine.MaterialExpressionPanner]
++MenuCategories=NSLOCTEXT("MaterialExpression", "Coordinates", "Coordinates")
+
+[/Script/Engine.MaterialExpressionParticleColor]
++MenuCategories=NSLOCTEXT("MaterialExpression", "Particles", "Particles")
++MenuCategories=NSLOCTEXT("MaterialExpression", "Constants", "Constants")
+bShaderInputData=true
+
+[/Script/Engine.MaterialExpressionParticleDirection]
++MenuCategories=NSLOCTEXT("MaterialExpression", "Particles", "Particles")
++MenuCategories=NSLOCTEXT("MaterialExpression", "Constants", "Constants")
+bShaderInputData=true
+
+[/Script/Engine.MaterialExpressionParticleMacroUV]
++MenuCategories=NSLOCTEXT("MaterialExpression", "Particles", "Particles")
+bShaderInputData=true
+
+[/Script/Engine.MaterialExpressionParticleMotionBlurFade]
++MenuCategories=NSLOCTEXT("MaterialExpression", "Particles", "Particles")
++MenuCategories=NSLOCTEXT("MaterialExpression", "Constants", "Constants")
+bShaderInputData=true
+
+[/Script/Engine.MaterialExpressionParticlePositionWS]
++MenuCategories=NSLOCTEXT("MaterialExpression", "Particles", "Particles")
++MenuCategories=NSLOCTEXT("MaterialExpression", "Coordinates", "Coordinates")
++MenuCategories=NSLOCTEXT("MaterialExpression", "Vectors", "Vectors")
+bShaderInputData=true
+
+[/Script/Engine.MaterialExpressionParticleRadius]
++MenuCategories=NSLOCTEXT("MaterialExpression", "Particles", "Particles")
++MenuCategories=NSLOCTEXT("MaterialExpression", "Constants", "Constants")
+bShaderInputData=true
+
+[/Script/Engine.MaterialExpressionParticleRandom]
++MenuCategories=NSLOCTEXT("MaterialExpression", "Particles", "Particles")
++MenuCategories=NSLOCTEXT("MaterialExpression", "Constants", "Constants")
+bShaderInputData=true
+
+[/Script/Engine.MaterialExpressionParticleRelativeTime]
++MenuCategories=NSLOCTEXT("MaterialExpression", "Particles", "Particles")
++MenuCategories=NSLOCTEXT("MaterialExpression", "Constants", "Constants")
+bShaderInputData=true
+
+[/Script/Engine.MaterialExpressionParticleSize]
++MenuCategories=NSLOCTEXT("MaterialExpression", "Particles", "Particles")
++MenuCategories=NSLOCTEXT("MaterialExpression", "Constants", "Constants")
+bShaderInputData=true
+
+[/Script/Engine.MaterialExpressionParticleSpeed]
++MenuCategories=NSLOCTEXT("MaterialExpression", "Particles", "Particles")
++MenuCategories=NSLOCTEXT("MaterialExpression", "Constants", "Constants")
+bShaderInputData=true
+
+[/Script/Engine.MaterialExpressionParticleSpriteRotation]
++MenuCategories=NSLOCTEXT("MaterialExpression", "Particles", "Particles")
++MenuCategories=NSLOCTEXT("MaterialExpression", "Constants", "Constants")
+bShaderInputData=true
+
+[/Script/Engine.MaterialExpressionParticleSubUV]
++MenuCategories=NSLOCTEXT("MaterialExpression", "Texture", "Texture")
++MenuCategories=NSLOCTEXT("MaterialExpression", "Particles", "Particles")
+
+[/Script/Engine.MaterialExpressionParticleSubUVProperties]
++MenuCategories=NSLOCTEXT("MaterialExpression", "Particles", "Particles")
++MenuCategories=NSLOCTEXT("MaterialExpression", "Coordinates", "Coordinates")
+bShaderInputData=true
+
+[/Script/Engine.MaterialExpressionPathTracingBufferTexture]
++MenuCategories=NSLOCTEXT("MaterialExpression", "PathTracing", "PathTracing")
+bShaderInputData=true
+
+[/Script/Engine.MaterialExpressionPathTracingQualitySwitch]
++MenuCategories=NSLOCTEXT("MaterialExpression", "Utility", "Utility")
+
+[/Script/Engine.MaterialExpressionPathTracingRayTypeSwitch]
++MenuCategories=NSLOCTEXT("MaterialExpression", "Utility", "Utility")
+
+[/Script/Engine.MaterialExpressionPerInstanceCustomData]
++MenuCategories=NSLOCTEXT("MaterialExpression", "Custom", "Custom")
+bShaderInputData=true
+
+[/Script/Engine.MaterialExpressionPerInstanceCustomData3Vector]
++MenuCategories=NSLOCTEXT("MaterialExpression", "Custom", "Custom")
+bShaderInputData=true
+
+[/Script/Engine.MaterialExpressionPerInstanceFadeAmount]
++MenuCategories=NSLOCTEXT("MaterialExpression", "Constants", "Constants")
+bShaderInputData=true
+
+[/Script/Engine.MaterialExpressionPerInstanceRandom]
++MenuCategories=NSLOCTEXT("MaterialExpression", "Constants", "Constants")
+bShaderInputData=true
+
+[/Script/Engine.MaterialExpressionPinBase]
++MenuCategories=NSLOCTEXT("MaterialExpression", "Utility", "Utility")
+
+[/Script/Engine.MaterialExpressionPixelDepth]
++MenuCategories=NSLOCTEXT("MaterialExpression", "Depth", "Depth")
+bShaderInputData=true
+
+[/Script/Engine.MaterialExpressionPixelNormalWS]
++MenuCategories=NSLOCTEXT("MaterialExpression", "Vectors", "Vectors")
++MenuCategories=NSLOCTEXT("MaterialExpression", "Coordinates", "Coordinates")
+bShaderInputData=true
+
+[/Script/Engine.MaterialExpressionPostVolumeUserFlagTest]
++MenuCategories=NSLOCTEXT("MaterialExpression", "Get Post Process Setting", "Get Post Process Setting")
+
+[/Script/Engine.MaterialExpressionPower]
++MenuCategories=NSLOCTEXT("MaterialExpression", "Math", "Math")
+
+[/Script/Engine.MaterialExpressionPrecomputedAOMask]
++MenuCategories=NSLOCTEXT("MaterialExpression", "Constants", "Constants")
+bShaderInputData=true
+
+[/Script/Engine.MaterialExpressionPreSkinnedLocalBounds]
++MenuCategories=NSLOCTEXT("MaterialExpression", "Vectors", "Vectors")
+bShaderInputData=true
+
+[/Script/Engine.MaterialExpressionPreSkinnedNormal]
++MenuCategories=NSLOCTEXT("MaterialExpression", "Vectors", "Vectors")
+bShaderInputData=true
+
+[/Script/Engine.MaterialExpressionPreSkinnedPosition]
++MenuCategories=NSLOCTEXT("MaterialExpression", "Vectors", "Vectors")
+bShaderInputData=true
+
+[/Script/Engine.MaterialExpressionPreviousFrameSwitch]
++MenuCategories=NSLOCTEXT("MaterialExpression", "Functions", "Functions")
+
+[/Script/Engine.MaterialExpressionQualitySwitch]
++MenuCategories=NSLOCTEXT("MaterialExpression", "Utility", "Utility")
+
+[/Script/Engine.MaterialExpressionRayTracingQualitySwitch]
++MenuCategories=NSLOCTEXT("MaterialExpression", "Utility", "Utility")
+
+[/Script/Engine.MaterialExpressionReflectionCapturePassSwitch]
++MenuCategories=NSLOCTEXT("MaterialExpression", "Utility", "Utility")
+
+[/Script/Engine.MaterialExpressionReflectionVectorWS]
++MenuCategories=NSLOCTEXT("MaterialExpression", "Vectors", "Vectors")
+bShaderInputData=true
+
+[/Script/Engine.MaterialExpressionReroute]
++MenuCategories=NSLOCTEXT("MaterialExpression", "Utility", "Utility")
+
+[/Script/Engine.MaterialExpressionRequiredSamplersSwitch]
++MenuCategories=NSLOCTEXT("MaterialExpression", "Utility", "Utility")
+
+[/Script/Engine.MaterialExpressionRgbToHsv]
++MenuCategories=NSLOCTEXT("MaterialExpression", "Math", "Math")
+
+[/Script/Engine.MaterialExpressionRotateAboutAxis]
++MenuCategories=NSLOCTEXT("MaterialExpression", "Utility", "Utility")
+
+[/Script/Engine.MaterialExpressionRotator]
++MenuCategories=NSLOCTEXT("MaterialExpression", "Coordinates", "Coordinates")
+
+[/Script/Engine.MaterialExpressionRound]
++MenuCategories=NSLOCTEXT("MaterialExpression", "Math", "Math")
+
+[/Script/Engine.MaterialExpressionRuntimeVirtualTextureCustomData]
++MenuCategories=NSLOCTEXT("MaterialExpression", "VirtualTexture", "VirtualTexture")
+
+[/Script/Engine.MaterialExpressionRuntimeVirtualTextureOutput]
++MenuCategories=NSLOCTEXT("MaterialExpression", "VirtualTexture", "VirtualTexture")
+
+[/Script/Engine.MaterialExpressionRuntimeVirtualTextureReplace]
++MenuCategories=NSLOCTEXT("MaterialExpression", "VirtualTexture", "VirtualTexture")
+
+[/Script/Engine.MaterialExpressionRuntimeVirtualTextureSample]
++MenuCategories=NSLOCTEXT("MaterialExpression", "VirtualTexture", "VirtualTexture")
+
+[/Script/Engine.MaterialExpressionRuntimeVirtualTextureSampleParameter]
++MenuCategories=NSLOCTEXT("MaterialExpression", "VirtualTexture", "VirtualTexture")
++MenuCategories=NSLOCTEXT("MaterialExpression", "Parameters", "Parameters")
+
+[/Script/Engine.MaterialExpressionSamplePhysicsVectorField]
++MenuCategories=NSLOCTEXT("MaterialExpression", "Utility", "Utility")
+
+[/Script/Engine.MaterialExpressionSamplePhysicsScalarField]
++MenuCategories=NSLOCTEXT("MaterialExpression", "Utility", "Utility")
+
+[/Script/Engine.MaterialExpressionSamplePhysicsIntegerField]
++MenuCategories=NSLOCTEXT("MaterialExpression", "Utility", "Utility")
+
+[/Script/Engine.MaterialExpressionSaturate]
++MenuCategories=NSLOCTEXT("MaterialExpression", "Math", "Math")
+
+[/Script/Engine.MaterialExpressionSceneColor]
++MenuCategories=NSLOCTEXT("MaterialExpression", "Texture", "Texture")
+bShaderInputData=true
+
+[/Script/Engine.MaterialExpressionSceneDepth]
++MenuCategories=NSLOCTEXT("MaterialExpression", "Depth", "Depth")
+bShaderInputData=true
+
+[/Script/Engine.MaterialExpressionSceneDepthWithoutWater]
++MenuCategories=NSLOCTEXT("MaterialExpression", "Water", "Water")
+bShaderInputData=true
+
+[/Script/Engine.MaterialExpressionSceneTexelSize]
++MenuCategories=NSLOCTEXT("MaterialExpression", "Coordinates", "Coordinates")
+bShaderInputData=true
+
+[/Script/Engine.MaterialExpressionSceneTexture]
++MenuCategories=NSLOCTEXT("MaterialExpression", "Texture", "Texture")
+bShaderInputData=true
+
+[/Script/Engine.MaterialExpressionScreenPosition]
++MenuCategories=NSLOCTEXT("MaterialExpression", "Coordinates", "Coordinates")
+bShaderInputData=true
+
+[/Script/Engine.MaterialExpressionSetMaterialAttributes]
++MenuCategories=NSLOCTEXT("MaterialExpression", "MaterialAttributes", "Material Attributes")
+
+[/Script/Engine.MaterialExpressionShaderStageSwitch]
++MenuCategories=NSLOCTEXT("MaterialExpression", "Utility", "Utility")
+
+[/Script/Engine.MaterialExpressionShadingModel]
++MenuCategories=NSLOCTEXT("MaterialExpression", "Shading Model", "Shading Model")
+
+[/Script/Engine.MaterialExpressionShadingPathSwitch]
++MenuCategories=NSLOCTEXT("MaterialExpression", "Utility", "Utility")
+
+[/Script/Engine.MaterialExpressionShadowReplace]
++MenuCategories=NSLOCTEXT("MaterialExpression", "Utility", "Utility")
+
+[/Script/Engine.MaterialExpressionSign]
++MenuCategories=NSLOCTEXT("MaterialExpression", "Math", "Math")
+
+[/Script/Engine.MaterialExpressionSine]
++MenuCategories=NSLOCTEXT("MaterialExpression", "Math", "Math")
+
+[/Script/Engine.MaterialExpressionSingleLayerWaterMaterialOutput]
++MenuCategories=NSLOCTEXT("MaterialExpression", "Water", "Water")
+
+[/Script/Engine.MaterialExpressionSkyAtmosphereLightDirection]
++MenuCategories=NSLOCTEXT("MaterialExpression", "Sky", "Sky")
+
+[/Script/Engine.MaterialExpressionSkyAtmosphereLightIlluminance]
++MenuCategories=NSLOCTEXT("MaterialExpression", "Sky", "Sky")
+
+[/Script/Engine.MaterialExpressionSkyAtmosphereLightIlluminanceOnGround]
++MenuCategories=NSLOCTEXT("MaterialExpression", "Sky", "Sky")
+
+[/Script/Engine.MaterialExpressionSkyAtmosphereLightDiskLuminance]
++MenuCategories=NSLOCTEXT("MaterialExpression", "Sky", "Sky")
+
+[/Script/Engine.MaterialExpressionSkyAtmosphereAerialPerspective]
++MenuCategories=NSLOCTEXT("MaterialExpression", "Sky", "Sky")
+
+[/Script/Engine.MaterialExpressionSkyAtmosphereDistantLightScatteredLuminance]
++MenuCategories=NSLOCTEXT("MaterialExpression", "Sky", "Sky")
+
+[/Script/Engine.MaterialExpressionSkyAtmosphereViewLuminance]
++MenuCategories=NSLOCTEXT("MaterialExpression", "Sky", "Sky")
+
+[/Script/Engine.MaterialExpressionSkyLightEnvMapSample]
++MenuCategories=NSLOCTEXT("MaterialExpression", "Sky", "Sky")
+
+[/Script/Engine.MaterialExpressionSobol]
++MenuCategories=NSLOCTEXT("MaterialExpression", "Utility", "Utility")
+
+[/Script/Engine.MaterialExpressionSmoothStep]
++MenuCategories=NSLOCTEXT("MaterialExpression", "Math", "Math")
++MenuCategories=NSLOCTEXT("MaterialExpression", "Utility", "Utility")
+
+[/Script/Engine.MaterialExpressionSparseVolumeTextureObject]
++MenuCategories=NSLOCTEXT("MaterialExpression", "Texture", "Texture")
++MenuCategories=NSLOCTEXT("MaterialExpression", "Functions", "Functions")
+
+[/Script/Engine.MaterialExpressionSparseVolumeTextureSample]
++MenuCategories=NSLOCTEXT("MaterialExpression", "Texture", "Texture")
+
+[/Script/Engine.MaterialExpressionSparseVolumeTextureSampleParameter]
++MenuCategories=NSLOCTEXT("MaterialExpression", "Texture", "Texture")
++MenuCategories=NSLOCTEXT("MaterialExpression", "Parameters", "Parameters")
+
+[/Script/Engine.MaterialExpressionSparseVolumeTextureObjectParameter]
++MenuCategories=NSLOCTEXT("MaterialExpression", "Texture", "Texture")
++MenuCategories=NSLOCTEXT("MaterialExpression", "Parameters", "Parameters")
+
+[/Script/Engine.MaterialExpressionSpeedTree]
++MenuCategories=NSLOCTEXT("MaterialExpression", "SpeedTree", "SpeedTree")
+
+[/Script/Engine.MaterialExpressionSphereMask]
++MenuCategories=NSLOCTEXT("MaterialExpression", "Utility", "Utility")
+
+[/Script/Engine.MaterialExpressionSphericalParticleOpacity]
++MenuCategories=NSLOCTEXT("MaterialExpression", "Particles", "Particles")
+
+[/Script/Engine.MaterialExpressionSquareRoot]
++MenuCategories=NSLOCTEXT("MaterialExpression", "Math", "Math")
+
+[/Script/Engine.MaterialExpressionSRGBColorToWorkingColorSpace]
++MenuCategories=NSLOCTEXT("MaterialExpression", "Color", "Color")
+
+[/Script/Engine.MaterialExpressionStaticBool]
++MenuCategories=NSLOCTEXT("MaterialExpression", "Functions", "Functions")
+bShaderInputData=true
+
+[/Script/Engine.MaterialExpressionStaticBoolParameter]
++MenuCategories=NSLOCTEXT("MaterialExpression", "Parameters", "Parameters")
+
+[/Script/Engine.MaterialExpressionStaticComponentMaskParameter]
++MenuCategories=NSLOCTEXT("MaterialExpression", "Parameters", "Parameters")
++MenuCategories=NSLOCTEXT("MaterialExpression", "Parameters", "Parameters")
+
+[/Script/Engine.MaterialExpressionStaticSwitch]
++MenuCategories=NSLOCTEXT("MaterialExpression", "Functions", "Functions")
+
+[/Script/Engine.MaterialExpressionStaticSwitchParameter]
++MenuCategories=NSLOCTEXT("MaterialExpression", "Parameters", "Parameters")
+
+[/Script/Engine.MaterialExpressionStep]
++MenuCategories=NSLOCTEXT("MaterialExpression", "Math", "Math")
++MenuCategories=NSLOCTEXT("MaterialExpression", "Utility", "Utility")
+
+[/Script/Engine.MaterialExpressionSubsurfaceMediumMaterialOutput]
++MenuCategories=NSLOCTEXT("MaterialExpression", "Shading", "Shading")
+
+[/Script/Engine.MaterialExpressionSubtract]
++MenuCategories=NSLOCTEXT("MaterialExpression", "Math", "Math")
+
+[/Script/Engine.MaterialExpressionSwitch]
++MenuCategories=NSLOCTEXT("MaterialExpression", "Math", "Math")
+
+[/Script/Engine.MaterialExpressionTangent]
++MenuCategories=NSLOCTEXT("MaterialExpression", "Math", "Math")
+
+[/Script/Engine.MaterialExpressionTemporalSobol]
++MenuCategories=NSLOCTEXT("MaterialExpression", "Utility", "Utility")
+
+[/Script/Engine.MaterialExpressionTextureCoordinate]
++MenuCategories=NSLOCTEXT("MaterialExpression", "Coordinates", "Coordinates")
+
+[/Script/Engine.MaterialExpressionTextureCollection]
++MenuCategories=NSLOCTEXT("MaterialExpression", "Texture", "Texture")
+
+[/Script/Engine.MaterialExpressionTextureCollectionParameter]
++MenuCategories=NSLOCTEXT("MaterialExpression", "Texture", "Texture")
++MenuCategories=NSLOCTEXT("MaterialExpression", "Parameters", "Parameters")
+
+[/Script/Engine.MaterialExpressionTextureObject]
++MenuCategories=NSLOCTEXT("MaterialExpression", "Texture", "Texture")
++MenuCategories=NSLOCTEXT("MaterialExpression", "Functions", "Functions")
+
+[/Script/Engine.MaterialExpressionTextureObjectFromCollection]
++MenuCategories=NSLOCTEXT("MaterialExpression", "Texture", "Texture")
+
+[/Script/Engine.MaterialExpressionTextureObjectParameter]
++MenuCategories=NSLOCTEXT("MaterialExpression", "Texture", "Texture")
++MenuCategories=NSLOCTEXT("MaterialExpression", "Parameters", "Parameters")
+
+[/Script/Engine.MaterialExpressionTextureProperty]
++MenuCategories=NSLOCTEXT("MaterialExpression", "Texture", "Texture")
+bShaderInputData=true
+
+[/Script/Engine.MaterialExpressionTextureSampleParameter2DArray]
++MenuCategories=NSLOCTEXT("MaterialExpression", "Texture", "Texture")
++MenuCategories=NSLOCTEXT("MaterialExpression", "Parameters", "Parameters")
+
+[/Script/Engine.MaterialExpressionTextureSampleParameterCube]
++MenuCategories=NSLOCTEXT("MaterialExpression", "Texture", "Texture")
++MenuCategories=NSLOCTEXT("MaterialExpression", "Parameters", "Parameters")
+
+[/Script/Engine.MaterialExpressionTextureSampleParameterCubeArray]
++MenuCategories=NSLOCTEXT("MaterialExpression", "Texture", "Texture")
++MenuCategories=NSLOCTEXT("MaterialExpression", "Parameters", "Parameters")
+
+[/Script/Engine.MaterialExpressionTextureSampleParameterSubUV]
++MenuCategories=NSLOCTEXT("MaterialExpression", "Texture", "Texture")
++MenuCategories=NSLOCTEXT("MaterialExpression", "Parameters", "Parameters")
++MenuCategories=NSLOCTEXT("MaterialExpression", "Particles", "Particles")
+
+[/Script/Engine.MaterialExpressionTextureSampleParameterVolume]
++MenuCategories=NSLOCTEXT("MaterialExpression", "Texture", "Texture")
++MenuCategories=NSLOCTEXT("MaterialExpression", "Parameters", "Parameters")
+
+[/Script/Engine.MaterialExpressionTime]
++MenuCategories=NSLOCTEXT("MaterialExpression", "Constants", "Constants")
+bShaderInputData=true
+
+[/Script/Engine.MaterialExpressionTransform]
++MenuCategories=NSLOCTEXT("MaterialExpression", "VectorOps", "VectorOps")
+
+[/Script/Engine.MaterialExpressionTransformPosition]
++MenuCategories=NSLOCTEXT("MaterialExpression", "VectorOps", "VectorOps")
+
+[/Script/Engine.MaterialExpressionTruncate]
++MenuCategories=NSLOCTEXT("MaterialExpression", "Math", "Math")
+
+[/Script/Engine.MaterialExpressionTruncateLWC]
++MenuCategories=NSLOCTEXT("MaterialExpression", "Coordinates", "Coordinates")
+
+[/Script/Engine.MaterialExpressionTwoSidedSign]
++MenuCategories=NSLOCTEXT("MaterialExpression", "Constants", "Constants")
+bShaderInputData=true
+
+[/Script/Engine.MaterialExpressionUserSceneTexture]
++MenuCategories=NSLOCTEXT("MaterialExpression", "Texture", "Texture")
+bShaderInputData=true
+
+[/Script/Engine.MaterialExpressionVectorNoise]
++MenuCategories=NSLOCTEXT("MaterialExpressionNoise", "Utility", "Utility")
+
+[/Script/Engine.MaterialExpressionVertexColor]
++MenuCategories=NSLOCTEXT("MaterialExpression", "Constants", "Constants")
+bShaderInputData=true
+
+[/Script/Engine.MaterialExpressionVertexInterpolator]
++MenuCategories=NSLOCTEXT("MaterialExpression", "Utility", "Utility")
+
+[/Script/Engine.MaterialExpressionVertexNormalWS]
++MenuCategories=NSLOCTEXT("MaterialExpression", "Vectors", "Vectors")
++MenuCategories=NSLOCTEXT("MaterialExpression", "Coordinates", "Coordinates")
+bShaderInputData=true
+
+[/Script/Engine.MaterialExpressionVertexTangentWS]
++MenuCategories=NSLOCTEXT("MaterialExpression", "Vectors", "Vectors")
++MenuCategories=NSLOCTEXT("MaterialExpression", "Coordinates", "Coordinates")
+bShaderInputData=true
+
+[/Script/Engine.MaterialExpressionViewProperty]
++MenuCategories=NSLOCTEXT("MaterialExpression", "Constants", "Constants")
+bShaderInputData=true
+
+[/Script/Engine.MaterialExpressionViewSize]
++MenuCategories=NSLOCTEXT("MaterialExpression", "Coordinates", "Coordinates")
+bShaderInputData=true
+
+[/Script/Engine.MaterialExpressionVirtualTextureFeatureSwitch]
++MenuCategories=NSLOCTEXT("MaterialExpression", "VirtualTexture", "VirtualTexture")
+
+[/Script/Engine.MaterialExpressionVolumetricAdvancedMaterialInput]
++MenuCategories=NSLOCTEXT("MaterialExpression", "Volume", "Volume")
+
+[/Script/Engine.MaterialExpressionVolumetricCloudEmptySpaceSkippingInput]
++MenuCategories=NSLOCTEXT("MaterialExpression", "Volume", "Volume")
+
+[/Script/Engine.MaterialExpressionVolumetricAdvancedMaterialOutput]
++MenuCategories=NSLOCTEXT("MaterialExpression", "Volume", "Volume")
+
+[/Script/Engine.MaterialExpressionVolumetricCloudEmptySpaceSkippingOutput]
++MenuCategories=NSLOCTEXT("MaterialExpression", "Volume", "Volume")
+
+[/Script/Engine.MaterialExpressionWorldPosition]
++MenuCategories=NSLOCTEXT("MaterialExpression", "Coordinates", "Coordinates")
+bShaderInputData=true
+
+[/Script/Paper2D.MaterialExpressionSpriteTextureSampler]
++MenuCategories=NSLOCTEXT("MaterialExpression", "Texture", "Texture")
++MenuCategories=NSLOCTEXT("MaterialExpression", "Parameters", "Parameters")
\ No newline at end of file
diff --git a/Engine/Config/BaseMaterialExpressions.ini b/Engine/Config/BaseMaterialExpressions.ini
deleted file mode 100644
diff --git a/Engine/Source/Runtime/Engine/Private/Materials/HLSLMaterialDerivativeAutogen.h b/Engine/Source/Runtime/Engine/Private/Materials/HLSLMaterialDerivativeAutogen.h
--- a/Engine/Source/Runtime/Engine/Private/Materials/HLSLMaterialDerivativeAutogen.h
+++ b/Engine/Source/Runtime/Engine/Private/Materials/HLSLMaterialDerivativeAutogen.h
@@ -5,7 +5,7 @@
 #include "CoreMinimal.h"
 #include "MaterialShared.h"
 
-#include "HLSLMaterialDerivativeAutogen.generated.h"
+#if WITH_EDITORONLY_DATA
 
 class FHLSLMaterialTranslator;
 
@@ -16,8 +16,7 @@ class FHLSLMaterialTranslator;
 * Zero		- This node is aware of partial derivatives, and knows that it's value is zero, as is the case for uniform parameters.
 * Valid		- This node is aware of partial derivatives, and knows that it has a calculated value.
 **/
-UENUM()
-enum class EDerivativeStatus : uint8
+enum class EDerivativeStatus
 {
 	NotAware,
 	NotValid,
@@ -25,8 +24,6 @@ enum class EDerivativeStatus : uint8
 	Valid,
 };
 
-#if WITH_EDITORONLY_DATA
-
 inline bool IsDerivativeValid(EDerivativeStatus Status)
 {
 	return Status == EDerivativeStatus::Valid || Status == EDerivativeStatus::Zero;
diff --git a/Engine/Source/Runtime/Engine/Private/Materials/HLSLMaterialTranslator.cpp b/Engine/Source/Runtime/Engine/Private/Materials/HLSLMaterialTranslator.cpp
--- a/Engine/Source/Runtime/Engine/Private/Materials/HLSLMaterialTranslator.cpp
+++ b/Engine/Source/Runtime/Engine/Private/Materials/HLSLMaterialTranslator.cpp
@@ -36,7 +36,6 @@
 #include "Materials/MaterialExpressionVolumetricAdvancedMaterialOutput.h"
 #include "Materials/MaterialExpressionWorldPosition.h"
 #include "Materials/MaterialExpressionFirstPersonOutput.h"
-#include "Materials/MaterialExternalCodeRegistry.h"
 #include "Materials/MaterialSourceTemplate.h"
 #include "Materials/SubstrateMaterial.h"
 #include "StringTemplate.h"
@@ -3832,25 +3831,6 @@ int32 FHLSLMaterialTranslator::AddCodeChunkWithHash(uint64 BaseHash, EMaterialVa
 	return AddCodeChunkInner(Hash, FormattedCode, Type, EDerivativeStatus::NotAware, false);
 }
 
-int32 FHLSLMaterialTranslator::AddExternalCodeChunk(const FName& InExternalCodeIdentifier)
-{
-	if (const FMaterialExternalCodeDeclaration* ExternalCodeDeclaration = MaterialExternalCodeRegistry::Get().FindExternalCode(InExternalCodeIdentifier))
-	{
-		return AddExternalCodeChunk(*ExternalCodeDeclaration);
-	}
-	return Errorf(TEXT("External code identifier '%s' not found"), *InExternalCodeIdentifier.ToString());
-}
-
-int32 FHLSLMaterialTranslator::AddExternalCodeChunk(const FMaterialExternalCodeDeclaration& InExternalCode)
-{
-	if (((int32)InExternalCode.ShaderFrequency & (1 << (int32)ShaderFrequency)) == 0)
-	{
-		FName ShaderFrequencyName = StaticEnum<EMaterialShaderFrequency>()->GetNameByValue(1ull << (int32)ShaderFrequency);
-		return Errorf(TEXT("External code identifier '%s' is not available in the %s shader."), *InExternalCode.Name.ToString(), *ShaderFrequencyName.ToString());
-	}
-	return AddCodeChunkInner(InExternalCode.GetReturnTypeValue(), InExternalCode.Derivative, InExternalCode.bIsInlined, TEXT("%s"), *InExternalCode.Definition);
-}
-
 int32 FHLSLMaterialTranslator::AddUniformExpressionInner(uint64 Hash, FMaterialUniformExpression* UniformExpression, EMaterialValueType Type, const TCHAR* FormattedCode)
 {
 	check(bAllowCodeChunkGeneration);
@@ -5392,7 +5372,7 @@ int32 FHLSLMaterialTranslator::ViewProperty(EMaterialExposedViewProperty Propert
 
 int32 FHLSLMaterialTranslator::IsOrthographic()
 {
-	return AddExternalCodeChunk(TEXT("IsOrthographic"));
+	return AddInlinedCodeChunkZeroDeriv(MCT_Float, TEXT("IsOrthoProjectionFloat()"));
 }
 
 
@@ -6001,11 +5981,15 @@ int32 FHLSLMaterialTranslator::Abs(int32 X)
 
 int32 FHLSLMaterialTranslator::ReflectionVector()
 {
+	if (ShaderFrequency != SF_Pixel && ShaderFrequency != SF_Compute)
+	{
+		return NonPixelShaderExpressionError();
+	}
 	if (ShaderFrequency != SF_Vertex)
 	{
 		bUsesTransformVector = true;
 	}
-	return AddExternalCodeChunk(TEXT("ReflectionVector"));
+	return AddInlinedCodeChunk(MCT_Float3,TEXT("Parameters.ReflectionVector"));
 }
 
 int32 FHLSLMaterialTranslator::ReflectionAboutCustomWorldNormal(int32 CustomWorldNormal, int32 bNormalizeCustomWorldNormal)
@@ -6032,20 +6016,30 @@ int32 FHLSLMaterialTranslator::ReflectionAboutCustomWorldNormal(int32 CustomWorl
 
 int32 FHLSLMaterialTranslator::CameraVector()
 {
+	if (ShaderFrequency != SF_Pixel && ShaderFrequency != SF_Compute)
+	{
+		return NonPixelShaderExpressionError();
+	}
 	if (ShaderFrequency != SF_Vertex)
 	{
 		bUsesTransformVector = true;
 	}
-	return AddExternalCodeChunk(TEXT("CameraVector"));
+	return AddInlinedCodeChunk(MCT_Float3,TEXT("Parameters.CameraVector"));
 }
 
 int32 FHLSLMaterialTranslator::LightVector()
 {
+	if (ShaderFrequency != SF_Pixel && ShaderFrequency != SF_Compute)
+	{
+		return NonPixelShaderExpressionError();
+	}
+
 	if (!Material->IsLightFunction() && !Material->IsDeferredDecal())
 	{
 		return Errorf(TEXT("LightVector can only be used in LightFunction or DeferredDecal materials"));
 	}
-	return AddExternalCodeChunk(TEXT("LightVector"));
+
+	return AddInlinedCodeChunk(MCT_Float3,TEXT("Parameters.LightVector"));
 }
 
 int32 FHLSLMaterialTranslator::GetViewportUV()
@@ -6088,7 +6082,12 @@ int32 FHLSLMaterialTranslator::GetPixelPosition()
 
 int32 FHLSLMaterialTranslator::ParticleMacroUV()
 {
-	return AddExternalCodeChunk(TEXT("ParticleMacroUV"));
+	if (ShaderFrequency != SF_Pixel && ShaderFrequency != SF_Compute)
+	{
+		return NonPixelShaderExpressionError();
+	}
+
+	return AddCodeChunk(MCT_Float2,TEXT("GetParticleMacroUV(Parameters)"));
 }
 
 int32 FHLSLMaterialTranslator::ParticleSubUV(int32 TextureIndex, EMaterialSamplerType SamplerType, int32 MipValue0Index, int32 MipValue1Index, ETextureMipValueMode MipValueMode, bool bBlend)
@@ -6155,8 +6154,12 @@ int32 FHLSLMaterialTranslator::ParticleSubUVProperty(int32 PropertyIndex)
 
 int32 FHLSLMaterialTranslator::ParticleColor()
 {
+	if (ShaderFrequency != SF_Vertex && ShaderFrequency != SF_Pixel && ShaderFrequency != SF_Compute)
+	{
+		return NonVertexOrPixelShaderExpressionError();
+	}
 	bUsesParticleColor |= (ShaderFrequency != SF_Vertex);
-	return AddExternalCodeChunk(TEXT("ParticleColor"));
+	return AddInlinedCodeChunkZeroDeriv(MCT_Float4,TEXT("Parameters.Particle.Color"));
 }
 
 int32 FHLSLMaterialTranslator::ParticlePosition(EPositionOrigin OriginType)
@@ -6189,8 +6192,12 @@ int32 FHLSLMaterialTranslator::ParticlePosition(EPositionOrigin OriginType)
 
 int32 FHLSLMaterialTranslator::ParticleRadius()
 {
+	if (ShaderFrequency != SF_Vertex && ShaderFrequency != SF_Pixel && ShaderFrequency != SF_Compute)
+	{
+		return NonVertexOrPixelShaderExpressionError();
+	}
 	bNeedsParticlePosition = true;
-	return AddExternalCodeChunk(TEXT("ParticleRadius"));
+	return AddInlinedCodeChunkZeroDeriv(MCT_Float,TEXT("max(Parameters.Particle.TranslatedWorldPositionAndSize.w, .001f)"));
 }
 
 int32 FHLSLMaterialTranslator::SphericalParticleOpacity(int32 Density)
@@ -6214,52 +6221,85 @@ int32 FHLSLMaterialTranslator::SphericalParticleOpacity(int32 Density)
 
 int32 FHLSLMaterialTranslator::ParticleRelativeTime()
 {
+	if (ShaderFrequency != SF_Vertex && ShaderFrequency != SF_Pixel && ShaderFrequency != SF_Compute)
+	{
+		return NonVertexOrPixelShaderExpressionError();
+	}
 	bNeedsParticleTime = true;
-	return AddExternalCodeChunk(TEXT("ParticleRelativeTime"));
+	return AddInlinedCodeChunkZeroDeriv(MCT_Float,TEXT("Parameters.Particle.RelativeTime"));
 }
 
 int32 FHLSLMaterialTranslator::ParticleMotionBlurFade()
 {
+	if (ShaderFrequency != SF_Vertex && ShaderFrequency != SF_Pixel && ShaderFrequency != SF_Compute)
+	{
+		return NonVertexOrPixelShaderExpressionError();
+	}
 	bUsesParticleMotionBlur = true;
-	return AddExternalCodeChunk(TEXT("ParticleMotionBlurFade"));
+	return AddInlinedCodeChunkZeroDeriv(MCT_Float,TEXT("Parameters.Particle.MotionBlurFade"));
 }
 
 int32 FHLSLMaterialTranslator::ParticleRandom()
 {
+	if (ShaderFrequency != SF_Vertex && ShaderFrequency != SF_Pixel && ShaderFrequency != SF_Compute)
+	{
+		return NonVertexOrPixelShaderExpressionError();
+	}
 	bNeedsParticleRandom = true;
-	return AddExternalCodeChunk(TEXT("ParticleRandom"));
+	return AddInlinedCodeChunkZeroDeriv(MCT_Float,TEXT("Parameters.Particle.Random"));
 }
 
 
 int32 FHLSLMaterialTranslator::ParticleDirection()
 {
+	if (ShaderFrequency != SF_Vertex && ShaderFrequency != SF_Pixel && ShaderFrequency != SF_Compute)
+	{
+		return NonVertexOrPixelShaderExpressionError();
+	}
 	bNeedsParticleVelocity = true;
-	return AddExternalCodeChunk(TEXT("ParticleDirection"));
+	return AddInlinedCodeChunkZeroDeriv(MCT_Float3,TEXT("Parameters.Particle.Velocity.xyz"));
 }
 
 int32 FHLSLMaterialTranslator::ParticleSpeed()
 {
+	if (ShaderFrequency != SF_Vertex && ShaderFrequency != SF_Pixel && ShaderFrequency != SF_Compute)
+	{
+		return NonVertexOrPixelShaderExpressionError();
+	}
 	bNeedsParticleVelocity = true;
-	return AddExternalCodeChunk(TEXT("ParticleSpeed"));
+	return AddInlinedCodeChunkZeroDeriv(MCT_Float,TEXT("Parameters.Particle.Velocity.w"));
 }
 
 int32 FHLSLMaterialTranslator::ParticleSize()
 {
+	if (ShaderFrequency != SF_Vertex && ShaderFrequency != SF_Pixel && ShaderFrequency != SF_Compute)
+	{
+		return NonVertexOrPixelShaderExpressionError();
+	}
 	bNeedsParticleSize = true;
-	return AddExternalCodeChunk(TEXT("ParticleSize"));
+	return AddInlinedCodeChunkZeroDeriv(MCT_Float2,TEXT("Parameters.Particle.Size"));
 }
 
 int32 FHLSLMaterialTranslator::ParticleSpriteRotation()
 {
+	if (ShaderFrequency != SF_Vertex && ShaderFrequency != SF_Pixel && ShaderFrequency != SF_Compute)
+	{
+		return NonVertexOrPixelShaderExpressionError();
+	}
 	bNeedsParticleSpriteRotation = true;
-	return AddExternalCodeChunk(TEXT("ParticleSpriteRotation"));
+	return AddInlinedCodeChunkZeroDeriv(MCT_Float2, TEXT("float2(Parameters.Particle.SpriteRotation, Parameters.Particle.SpriteRotation * 57.2957795131f)"));
 }
 
 int32 FHLSLMaterialTranslator::LocalPosition(EPositionIncludedOffsets IncludedOffsets, ELocalPositionOrigin OriginType)
 {
 	if (OriginType == ELocalPositionOrigin::InstancePreSkinning)
 	{
-		return AddExternalCodeChunk(TEXT("PreSkinnedPosition"));
+		if (ShaderFrequency != SF_Vertex)
+		{
+			return Errorf(TEXT("Pre-skinned position is only available in the vertex shader, pass through custom interpolators if needed."));
+		}
+
+		return AddInlinedCodeChunk(MCT_Float3,TEXT("Parameters.PreSkinnedPosition"));
 	}
 
 	// If compiling for the previous frame in the vertex shader
@@ -9002,14 +9042,13 @@ int32 FHLSLMaterialTranslator::StaticTerrainLayerWeight(FName LayerName,int32 De
 	}
 }
 
-int32 FHLSLMaterialTranslator::ExternalCode(const FMaterialExternalCodeDeclaration& InExternalCode)
-{
-	return AddExternalCodeChunk(InExternalCode);
-}
-
 int32 FHLSLMaterialTranslator::FontSignedDistanceData()
 {
-	return AddExternalCodeChunk(TEXT("FontSignedDistanceData"));
+	if (ShaderFrequency != SF_Pixel)
+	{
+		return Errorf(TEXT("Font Signed Distance node is only available in pixel shader input."));
+	}
+	return AddCodeChunk(MCT_Float4, TEXT("GetFontSignedDistanceData(Parameters)"));
 }
 
 int32 FHLSLMaterialTranslator::VertexColor()
@@ -9030,12 +9069,12 @@ int32 FHLSLMaterialTranslator::VertexColor()
 
 int32 FHLSLMaterialTranslator::MeshPaintTextureCoordinateIndex()
 {
-	return AddExternalCodeChunk(TEXT("MeshPaintTextureCoordinateIndex"));
+	return AddInlinedCodeChunkZeroDeriv(MCT_Float, TEXT("GetMeshPaintTextureCoordinateIndex(GetPrimitiveData(Parameters))"));
 }
 
 int32 FHLSLMaterialTranslator::MeshPaintTextureDescriptor()
 {
-	return AddExternalCodeChunk(TEXT("MeshPaintTextureDescriptor"));
+	return AddInlinedCodeChunkZeroDeriv(MCT_TextureMeshPaint, TEXT("GetMeshPaintTextureDescriptor(GetPrimitiveData(Parameters))"));
 }
 
 int32 FHLSLMaterialTranslator::MeshPaintTextureReplace(int32 Invalid, int32 Valid)
@@ -9045,12 +9084,22 @@ int32 FHLSLMaterialTranslator::MeshPaintTextureReplace(int32 Invalid, int32 Vali
 
 int32 FHLSLMaterialTranslator::PreSkinnedPosition()
 {
-	return AddExternalCodeChunk(TEXT("PreSkinnedPosition"));
+	if (ShaderFrequency != SF_Vertex)
+	{
+		return Errorf(TEXT("Pre-skinned position is only available in the vertex shader, pass through custom interpolators if needed."));
+	}
+
+	return AddInlinedCodeChunk(MCT_Float3,TEXT("Parameters.PreSkinnedPosition"));
 }
 
 int32 FHLSLMaterialTranslator::PreSkinnedNormal()
 {
-	return AddExternalCodeChunk(TEXT("PreSkinnedNormal"));
+	if (ShaderFrequency != SF_Vertex)
+	{
+		return Errorf(TEXT("Pre-skinned normal is only available in the vertex shader, pass through custom interpolators if needed."));
+	}
+
+	return AddInlinedCodeChunk(MCT_Float3,TEXT("Parameters.PreSkinnedNormal"));
 }
 
 int32 FHLSLMaterialTranslator::VertexInterpolator(uint32 InterpolatorIndex)
@@ -15818,9 +15867,14 @@ int32 FHLSLMaterialTranslator::SpeedTree(int32 GeometryArg, int32 WindArg, int32
 /**Experimental access to the EyeAdaptation RT for Post Process materials. Can be one frame behind depending on the value of BlendableLocation. */
 int32 FHLSLMaterialTranslator::EyeAdaptation()
 {
+	if( ShaderFrequency != SF_Pixel )
+	{
+		return NonPixelShaderExpressionError();
+	}
+
 	MaterialCompilationOutput.bUsesEyeAdaptation = true;
 
-	return AddExternalCodeChunk(TEXT("EyeAdaptation"));
+	return AddInlinedCodeChunkZeroDeriv(MCT_Float, TEXT("EyeAdaptationLookup()"));
 }
 
 /**Experimental access to the EyeAdaptation RT for applying an inverse. */
diff --git a/Engine/Source/Runtime/Engine/Private/Materials/HLSLMaterialTranslator.h b/Engine/Source/Runtime/Engine/Private/Materials/HLSLMaterialTranslator.h
--- a/Engine/Source/Runtime/Engine/Private/Materials/HLSLMaterialTranslator.h
+++ b/Engine/Source/Runtime/Engine/Private/Materials/HLSLMaterialTranslator.h
@@ -801,10 +801,6 @@ protected:
 		return AddCodeChunkInner(Type, EDerivativeStatus::NotValid, true, Format, Args...);
 	}
 
-	// Constructs a code chunk that comes from an external HLSL function. Its return type and function implementation is data-driven and defined in BaseEditor.ini.
-	int32 AddExternalCodeChunk(const FName& InExternalCodeIdentifier);
-	int32 AddExternalCodeChunk(const FMaterialExternalCodeDeclaration& InExternalCode);
-
 	int32 AddUniformExpressionInner(uint64 Hash, FMaterialUniformExpression* UniformExpression, EMaterialValueType Type, const TCHAR* FormattedCode);
 
 	// AddUniformExpression - Adds an input to the Code array and returns its index.
@@ -1134,8 +1130,6 @@ protected:
 
 	virtual int32 StaticTerrainLayerWeight(FName ParameterName, int32 Default, bool bTextureArray = false) override;
 
-	virtual int32 ExternalCode(const FMaterialExternalCodeDeclaration& InExternalCode) override;
-
 	virtual int32 FontSignedDistanceData() override;
 
 	virtual int32 VertexColor() override;
diff --git a/Engine/Source/Runtime/Engine/Private/Materials/MaterialExpressions.cpp b/Engine/Source/Runtime/Engine/Private/Materials/MaterialExpressions.cpp
--- a/Engine/Source/Runtime/Engine/Private/Materials/MaterialExpressions.cpp
+++ b/Engine/Source/Runtime/Engine/Private/Materials/MaterialExpressions.cpp
@@ -97,7 +97,6 @@
 #include "Materials/MaterialExpressionDivide.h"
 #include "Materials/MaterialExpressionDotProduct.h"
 #include "Materials/MaterialExpressionDynamicParameter.h"
-#include "Materials/MaterialExpressionExternalCodeBase.h"
 #include "Materials/MaterialExpressionCloudLayer.h"
 #include "Materials/MaterialExpressionCustomOutput.h"
 #include "Materials/MaterialExpressionEyeAdaptation.h"
@@ -295,7 +294,6 @@
 #include "Materials/MaterialExpressionSparseVolumeTextureObject.h"
 #include "Materials/MaterialExpressionSparseVolumeTextureSample.h"
 #include "Materials/MaterialExpressionIsFirstPerson.h"
-#include "Materials/MaterialExternalCodeRegistry.h"
 #include "EditorSupportDelegates.h"
 #if WITH_EDITOR
 #include "MaterialGraph/MaterialGraphNode_Comment.h"
@@ -3797,6 +3795,11 @@ void UMaterialExpressionMeshPaintTextureCoordinateIndex::GetExpressionToolTip(TA
 	ConvertToMultilineToolTip(TEXT("Get the Mesh Paint Texture UV coordinate index."), 40, OutToolTip);
 }
 
+int32 UMaterialExpressionMeshPaintTextureCoordinateIndex::Compile(class FMaterialCompiler* Compiler, int32 OutputIndex)
+{
+	return Compiler->MeshPaintTextureCoordinateIndex();
+}
+
 #endif // WITH_EDITOR
 
 UMaterialExpressionMeshPaintTextureObject::UMaterialExpressionMeshPaintTextureObject(const FObjectInitializer& ObjectInitializer)
@@ -3808,19 +3811,6 @@ UMaterialExpressionMeshPaintTextureObject::UMaterialExpressionMeshPaintTextureOb
 
 #if WITH_EDITOR
 
-int32 UMaterialExpressionExternalCodeBase::Compile(class FMaterialCompiler* Compiler, int32 OutputIndex)
-{
-	if (const FMaterialExternalCodeDeclaration* ExternalCode = MaterialExternalCodeRegistry::Get().FindExternalCode(ExternalCodeIdentifier))
-	{
-		return Compiler->ExternalCode(*ExternalCode);
-	}
-	return Compiler->Errorf(TEXT("Missing external code declaration for '%s'"), *ExternalCodeIdentifier.ToString());
-}
-
-#endif
-
-#if WITH_EDITOR
-
 void UMaterialExpressionMeshPaintTextureObject::GetCaption(TArray<FString>& OutCaptions) const
 {
 	OutCaptions.Add(TEXT("Mesh Paint Texture Object"));
@@ -3836,6 +3826,11 @@ uint32 UMaterialExpressionMeshPaintTextureObject::GetOutputType(int32 OutputInde
 	return MCT_Texture2D;
 }
 
+int32 UMaterialExpressionMeshPaintTextureObject::Compile(class FMaterialCompiler* Compiler, int32 OutputIndex)
+{
+	return Compiler->MeshPaintTextureDescriptor();
+}
+
 #endif // WITH_EDITOR
 
 UMaterialExpressionMeshPaintTextureReplace::UMaterialExpressionMeshPaintTextureReplace(const FObjectInitializer& ObjectInitializer)
@@ -9838,7 +9833,8 @@ int32 UMaterialExpressionFontSignedDistance::Compile(class FMaterialCompiler* Co
 	{
 		return Compiler->Errorf(TEXT("Font Signed Distance node is only compatible with UI materials"));
 	}
-	return Super::Compile(Compiler, OutputIndex);
+
+	return Compiler->FontSignedDistanceData();
 }
 
 void UMaterialExpressionFontSignedDistance::GetCaption(TArray<FString>& OutCaptions) const
@@ -10118,6 +10114,11 @@ void UMaterialExpressionParticleSubUVProperties::GetCaption(TArray<FString>& Out
 //	MaterialExpressionParticleMacroUV
 //
 #if WITH_EDITOR
+int32 UMaterialExpressionParticleMacroUV::Compile(class FMaterialCompiler* Compiler, int32 OutputIndex)
+{
+	return Compiler->ParticleMacroUV();
+}
+
 void UMaterialExpressionParticleMacroUV::GetCaption(TArray<FString>& OutCaptions) const
 {
 	OutCaptions.Add(TEXT("Particle MacroUV"));
@@ -10218,6 +10219,11 @@ void UMaterialExpressionViewSize::GetCaption(TArray<FString>& OutCaptions) const
 #endif // WITH_EDITOR
 
 #if WITH_EDITOR
+int32 UMaterialExpressionIsOrthographic::Compile(class FMaterialCompiler* Compiler, int32 OutputIndex)
+{
+	return Compiler->IsOrthographic();
+}
+
 void UMaterialExpressionIsOrthographic::GetCaption(TArray<FString>& OutCaptions) const
 {
 	OutCaptions.Add(TEXT("IsOrthographic"));
@@ -21539,7 +21545,7 @@ int32 UMaterialExpressionPreSkinnedPosition::Compile(class FMaterialCompiler* Co
 		return Compiler->Errorf(TEXT("Pre-skinned position is only available in the vertex shader, pass through custom interpolators if needed."));
 	}
 
-	return Super::Compile(Compiler, OutputIndex);
+	return Compiler->PreSkinnedPosition();
 }
 
 void UMaterialExpressionPreSkinnedPosition::GetCaption(TArray<FString>& OutCaptions) const
@@ -21579,6 +21585,11 @@ UMaterialExpressionPreSkinnedNormal::UMaterialExpressionPreSkinnedNormal(const F
 }
 
 #if WITH_EDITOR
+int32 UMaterialExpressionPreSkinnedNormal::Compile(class FMaterialCompiler* Compiler, int32 OutputIndex)
+{
+	return Compiler->PreSkinnedNormal();
+}
+
 void UMaterialExpressionPreSkinnedNormal::GetCaption(TArray<FString>& OutCaptions) const
 {
 	OutCaptions.Add(TEXT("Pre-Skinned Local Normal"));
diff --git a/Engine/Source/Runtime/Engine/Private/Materials/MaterialExternalCodeRegistry.cpp b/Engine/Source/Runtime/Engine/Private/Materials/MaterialExternalCodeRegistry.cpp
deleted file mode 100644
diff --git a/Engine/Source/Runtime/Engine/Private/Materials/MaterialExternalCodeRegistry.h b/Engine/Source/Runtime/Engine/Private/Materials/MaterialExternalCodeRegistry.h
deleted file mode 100644
diff --git a/Engine/Source/Runtime/Engine/Public/MaterialCompiler.h b/Engine/Source/Runtime/Engine/Public/MaterialCompiler.h
--- a/Engine/Source/Runtime/Engine/Public/MaterialCompiler.h
+++ b/Engine/Source/Runtime/Engine/Public/MaterialCompiler.h
@@ -33,7 +33,6 @@ class UTexture;
 class UTextureCollection;
 struct FMaterialParameterInfo;
 class USparseVolumeTexture;
-struct FMaterialExternalCodeDeclaration;
 
 enum EMaterialForceCastFlags
 {
@@ -483,8 +482,6 @@ public:
 	// Feedback the instanced material layers data after the UMaterialExpressionMaterialAttributeLayers is assembled
 	virtual void FeedbackMaterialLayersInstancedGraphFromCompilation(const FMaterialLayersFunctions* InLayers) {} // TODO: SHould be pure virtual
 
-	/** Adds code to the material that is defined by an external function definition. */
-	virtual int32 ExternalCode(const FMaterialExternalCodeDeclaration& InExternalCode) = 0;
 
 	virtual int32 VertexColor() = 0;
 	virtual int32 MeshPaintTextureCoordinateIndex() = 0;
@@ -1073,8 +1070,6 @@ public:
 	virtual bool GetStaticBoolValue(int32 BoolIndex, bool& bSucceeded) override { return Compiler->GetStaticBoolValue(BoolIndex, bSucceeded); }
 	virtual int32 StaticTerrainLayerWeight(FName ParameterName, int32 Default, bool bTextureArray = false) override { return Compiler->StaticTerrainLayerWeight(ParameterName, Default, bTextureArray); }
 
-	virtual int32 ExternalCode(const FMaterialExternalCodeDeclaration& InExternalCode) override { return Compiler->ExternalCode(InExternalCode); }
-
 	virtual int32 FontSignedDistanceData() override { return Compiler->FontSignedDistanceData(); }
 
 	virtual int32 VertexColor() override { return Compiler->VertexColor(); }
diff --git a/Engine/Source/Runtime/Engine/Public/Materials/MaterialExpression.h b/Engine/Source/Runtime/Engine/Public/Materials/MaterialExpression.h
--- a/Engine/Source/Runtime/Engine/Public/Materials/MaterialExpression.h
+++ b/Engine/Source/Runtime/Engine/Public/Materials/MaterialExpression.h
@@ -139,7 +139,7 @@ struct FMaterialExpressionCollection
 	TObjectPtr<UMaterialExpression> ExpressionExecEnd = nullptr;
 };
 
-UCLASS(abstract, Optional, BlueprintType, Config=MaterialExpressions, hidecategories=Object, MinimalAPI)
+UCLASS(abstract, Optional, BlueprintType, Config = "Editor", hidecategories=Object, MinimalAPI)
 class UMaterialExpression : public UObject
 {
 	GENERATED_UCLASS_BODY()
diff --git a/Engine/Source/Runtime/Engine/Public/Materials/MaterialExpressionExternalCodeBase.h b/Engine/Source/Runtime/Engine/Public/Materials/MaterialExpressionExternalCodeBase.h
deleted file mode 100644
diff --git a/Engine/Source/Runtime/Engine/Public/Materials/MaterialExpressionFontSignedDistance.h b/Engine/Source/Runtime/Engine/Public/Materials/MaterialExpressionFontSignedDistance.h
--- a/Engine/Source/Runtime/Engine/Public/Materials/MaterialExpressionFontSignedDistance.h
+++ b/Engine/Source/Runtime/Engine/Public/Materials/MaterialExpressionFontSignedDistance.h
@@ -3,11 +3,11 @@
 
 #include "CoreMinimal.h"
 #include "UObject/ObjectMacros.h"
-#include "Materials/MaterialExpressionExternalCodeBase.h"
+#include "Materials/MaterialExpression.h"
 #include "MaterialExpressionFontSignedDistance.generated.h"
 
 UCLASS(collapsecategories, hidecategories=Object, MinimalAPI)
-class UMaterialExpressionFontSignedDistance : public UMaterialExpressionExternalCodeBase
+class UMaterialExpressionFontSignedDistance : public UMaterialExpression
 {
 	GENERATED_UCLASS_BODY()
 
diff --git a/Engine/Source/Runtime/Engine/Public/Materials/MaterialExpressionIsOrthographic.h b/Engine/Source/Runtime/Engine/Public/Materials/MaterialExpressionIsOrthographic.h
--- a/Engine/Source/Runtime/Engine/Public/Materials/MaterialExpressionIsOrthographic.h
+++ b/Engine/Source/Runtime/Engine/Public/Materials/MaterialExpressionIsOrthographic.h
@@ -5,11 +5,11 @@
 
 #include "CoreMinimal.h"
 #include "UObject/ObjectMacros.h"
-#include "Materials/MaterialExpressionExternalCodeBase.h"
+#include "Materials/MaterialExpression.h"
 #include "MaterialExpressionIsOrthographic.generated.h"
 
 UCLASS(collapsecategories, hidecategories=Object)
-class UMaterialExpressionIsOrthographic : public UMaterialExpressionExternalCodeBase
+class UMaterialExpressionIsOrthographic : public UMaterialExpression
 {
 	GENERATED_BODY()
 
@@ -17,6 +17,7 @@ public:
 
 	//~ Begin UMaterialExpression Interface
 #if WITH_EDITOR
+	virtual int32 Compile(class FMaterialCompiler* Compiler, int32 OutputIndex) override;
 	virtual void GetCaption(TArray<FString>& OutCaptions) const override;
 #endif
 	//~ End UMaterialExpression Interface
diff --git a/Engine/Source/Runtime/Engine/Public/Materials/MaterialExpressionMeshPaintTextureCoordinateIndex.h b/Engine/Source/Runtime/Engine/Public/Materials/MaterialExpressionMeshPaintTextureCoordinateIndex.h
--- a/Engine/Source/Runtime/Engine/Public/Materials/MaterialExpressionMeshPaintTextureCoordinateIndex.h
+++ b/Engine/Source/Runtime/Engine/Public/Materials/MaterialExpressionMeshPaintTextureCoordinateIndex.h
@@ -2,12 +2,12 @@
 
 #pragma once
 
-#include "Materials/MaterialExpressionExternalCodeBase.h"
+#include "Materials/MaterialExpression.h"
 
 #include "MaterialExpressionMeshPaintTextureCoordinateIndex.generated.h"
 
 UCLASS(collapsecategories, hidecategories=Object, MinimalAPI)
-class UMaterialExpressionMeshPaintTextureCoordinateIndex : public UMaterialExpressionExternalCodeBase
+class UMaterialExpressionMeshPaintTextureCoordinateIndex : public UMaterialExpression
 {
 	GENERATED_UCLASS_BODY()
 
@@ -15,6 +15,7 @@ class UMaterialExpressionMeshPaintTextureCoordinateIndex : public UMaterialExpre
 #if WITH_EDITOR
 	virtual void GetCaption(TArray<FString>& OutCaptions) const override;
 	virtual void GetExpressionToolTip(TArray<FString>& OutToolTip) override;
+	virtual int32 Compile(class FMaterialCompiler* Compiler, int32 OutputIndex) override;
 #endif
 	//~ End UMaterialExpression Interface
 };
diff --git a/Engine/Source/Runtime/Engine/Public/Materials/MaterialExpressionMeshPaintTextureObject.h b/Engine/Source/Runtime/Engine/Public/Materials/MaterialExpressionMeshPaintTextureObject.h
--- a/Engine/Source/Runtime/Engine/Public/Materials/MaterialExpressionMeshPaintTextureObject.h
+++ b/Engine/Source/Runtime/Engine/Public/Materials/MaterialExpressionMeshPaintTextureObject.h
@@ -2,12 +2,12 @@
 
 #pragma once
 
-#include "Materials/MaterialExpressionExternalCodeBase.h"
+#include "Materials/MaterialExpression.h"
 
 #include "MaterialExpressionMeshPaintTextureObject.generated.h"
 
 UCLASS(collapsecategories, hidecategories=Object, MinimalAPI)
-class UMaterialExpressionMeshPaintTextureObject : public UMaterialExpressionExternalCodeBase
+class UMaterialExpressionMeshPaintTextureObject : public UMaterialExpression
 {
 	GENERATED_UCLASS_BODY()
 
@@ -16,6 +16,7 @@ class UMaterialExpressionMeshPaintTextureObject : public UMaterialExpressionExte
 	virtual void GetCaption(TArray<FString>& OutCaptions) const override;
 	virtual void GetExpressionToolTip(TArray<FString>& OutToolTip) override;
 	virtual uint32 GetOutputType(int32 OutputIndex) override;
+	virtual int32 Compile(class FMaterialCompiler* Compiler, int32 OutputIndex) override;
 #endif
 	//~ End UMaterialExpression Interface
 };
diff --git a/Engine/Source/Runtime/Engine/Public/Materials/MaterialExpressionParticleMacroUV.h b/Engine/Source/Runtime/Engine/Public/Materials/MaterialExpressionParticleMacroUV.h
--- a/Engine/Source/Runtime/Engine/Public/Materials/MaterialExpressionParticleMacroUV.h
+++ b/Engine/Source/Runtime/Engine/Public/Materials/MaterialExpressionParticleMacroUV.h
@@ -9,11 +9,11 @@
 
 #include "CoreMinimal.h"
 #include "UObject/ObjectMacros.h"
-#include "Materials/MaterialExpressionExternalCodeBase.h"
+#include "Materials/MaterialExpression.h"
 #include "MaterialExpressionParticleMacroUV.generated.h"
 
 UCLASS()
-class UMaterialExpressionParticleMacroUV : public UMaterialExpressionExternalCodeBase
+class UMaterialExpressionParticleMacroUV : public UMaterialExpression
 {
 	GENERATED_BODY()
 
@@ -21,6 +21,7 @@ public:
 
 	//~ Begin UMaterialExpression Interface
 #if WITH_EDITOR
+	virtual int32 Compile(class FMaterialCompiler* Compiler, int32 OutputIndex) override;
 	virtual void GetCaption(TArray<FString>& OutCaptions) const override;
 
 #endif
diff --git a/Engine/Source/Runtime/Engine/Public/Materials/MaterialExpressionPreSkinnedNormal.h b/Engine/Source/Runtime/Engine/Public/Materials/MaterialExpressionPreSkinnedNormal.h
--- a/Engine/Source/Runtime/Engine/Public/Materials/MaterialExpressionPreSkinnedNormal.h
+++ b/Engine/Source/Runtime/Engine/Public/Materials/MaterialExpressionPreSkinnedNormal.h
@@ -4,16 +4,17 @@
 
 #include "CoreMinimal.h"
 #include "UObject/ObjectMacros.h"
-#include "Materials/MaterialExpressionExternalCodeBase.h"
+#include "Materials/MaterialExpression.h"
 #include "MaterialExpressionPreSkinnedNormal.generated.h"
 
 UCLASS(collapsecategories, hidecategories=Object, MinimalAPI)
-class UMaterialExpressionPreSkinnedNormal : public UMaterialExpressionExternalCodeBase
+class UMaterialExpressionPreSkinnedNormal : public UMaterialExpression
 {
 	GENERATED_UCLASS_BODY()
 
 	//~ Begin UMaterialExpression Interface
 #if WITH_EDITOR
+	virtual int32 Compile(class FMaterialCompiler* Compiler, int32 OutputIndex) override;
 	virtual void GetCaption(TArray<FString>& OutCaptions) const override;
 	virtual void GetExpressionToolTip(TArray<FString>& OutToolTip) override;
 	virtual FText GetKeywords() const override {return FText::FromString(TEXT("normal"));}
diff --git a/Engine/Source/Runtime/Engine/Public/Materials/MaterialExpressionPreSkinnedPosition.h b/Engine/Source/Runtime/Engine/Public/Materials/MaterialExpressionPreSkinnedPosition.h
--- a/Engine/Source/Runtime/Engine/Public/Materials/MaterialExpressionPreSkinnedPosition.h
+++ b/Engine/Source/Runtime/Engine/Public/Materials/MaterialExpressionPreSkinnedPosition.h
@@ -4,11 +4,11 @@
 
 #include "CoreMinimal.h"
 #include "UObject/ObjectMacros.h"
-#include "Materials/MaterialExpressionExternalCodeBase.h"
+#include "Materials/MaterialExpression.h"
 #include "MaterialExpressionPreSkinnedPosition.generated.h"
 
 UCLASS(collapsecategories, hidecategories=Object, MinimalAPI)
-class UMaterialExpressionPreSkinnedPosition : public UMaterialExpressionExternalCodeBase
+class UMaterialExpressionPreSkinnedPosition : public UMaterialExpression
 {
 	GENERATED_UCLASS_BODY()
 
