Remove r.MobileHDR=0 from the Android_XR device profile because it causes crashes when Mobile HDR is enabled at cooktime.
This includes any time the mobile deferred renderer is enabled, since it automatically sets r.MobileHDR=1 for the cook. This is then mismatched with r.MobileHDR=0 set by the device profile at runtime.

diff --git a/Engine/Config/BaseDeviceProfiles.ini b/Engine/Config/BaseDeviceProfiles.ini
--- a/Engine/Config/BaseDeviceProfiles.ini
+++ b/Engine/Config/BaseDeviceProfiles.ini
@@ -1205,7 +1205,6 @@ BaseProfileName=Android_High
 DeviceType=Android
 BaseProfileName=Android_Mid
 +CVars=vr.MobileMultiView=1
-+CVars=r.MobileHDR=0
 +CVars=r.Mobile.AntiAliasing=3
 +CVars=r.ForceStartInVR=1
 
