[Materials] Convert more material expressions to data-driven external code expressions.

#rnx

diff --git a/Engine/Config/BaseMaterialExpressions.ini b/Engine/Config/BaseMaterialExpressions.ini
--- a/Engine/Config/BaseMaterialExpressions.ini
+++ b/Engine/Config/BaseMaterialExpressions.ini
@@ -96,10 +96,12 @@ bShaderInputData=true
 
 [/Script/Engine.MaterialExpressionAtmosphericLightColor]
 +MenuCategories=NSLOCTEXT("MaterialExpression", "Utility", "Utility")
+ExternalCodeIdentifier="AtmosphericLightColor"
 
 [/Script/Engine.MaterialExpressionAtmosphericLightVector]
 +MenuCategories=NSLOCTEXT("MaterialExpression", "Utility", "Utility")
 bShaderInputData=true
+ExternalCodeIdentifier="AtmosphericLightVector"
 
 [/Script/Engine.MaterialExpressionBindlessSwitch]
 +MenuCategories=NSLOCTEXT("MaterialExpression", "Utility", "Utility")
@@ -895,6 +897,7 @@ bShaderInputData=true
 [/Script/Engine.MaterialExpressionTwoSidedSign]
 +MenuCategories=NSLOCTEXT("MaterialExpression", "Constants", "Constants")
 bShaderInputData=true
+ExternalCodeIdentifier="TwoSidedSign"
 
 [/Script/Engine.MaterialExpressionUserSceneTexture]
 +MenuCategories=NSLOCTEXT("MaterialExpression", "Texture", "Texture")
@@ -972,4 +975,7 @@ bShaderInputData=true
 +ExternalCodeDeclarations=(ReturnType=Float3, Name="PreSkinnedPosition", Definition="Parameters.PreSkinnedPosition", ShaderFrequency=Vertex, bIsInlined=true)
 +ExternalCodeDeclarations=(ReturnType=Float3, Name="PreSkinnedNormal", Definition="Parameters.PreSkinnedNormal", ShaderFrequency=Vertex, bIsInlined=true)
 +ExternalCodeDeclarations=(ReturnType=Float3, Name="ReflectionVector", Definition="Parameters.ReflectionVector", ShaderFrequency=Pixel|Compute, bIsInlined=true)
++ExternalCodeDeclarations=(ReturnType=Float, Name="TwoSidedSign", Definition="Parameters.TwoSidedSign", ShaderFrequency=Pixel|Compute, Derivative=Zero, bIsInlined=true)
++ExternalCodeDeclarations=(ReturnType=Float3, Name="AtmosphericLightVector", Definition="MaterialExpressionAtmosphericLightVector(Parameters)")
++ExternalCodeDeclarations=(ReturnType=Float3, Name="AtmosphericLightColor", Definition="MaterialExpressionAtmosphericLightColor(Parameters)")
 
diff --git a/Engine/Source/Runtime/Engine/Private/Materials/HLSLMaterialTranslator.cpp b/Engine/Source/Runtime/Engine/Private/Materials/HLSLMaterialTranslator.cpp
--- a/Engine/Source/Runtime/Engine/Private/Materials/HLSLMaterialTranslator.cpp
+++ b/Engine/Source/Runtime/Engine/Private/Materials/HLSLMaterialTranslator.cpp
@@ -11576,11 +11576,7 @@ int32 FHLSLMaterialTranslator::RotateAboutAxis(int32 NormalizedRotationAxisAndAn
 
 int32 FHLSLMaterialTranslator::TwoSidedSign()
 {
-	if (ShaderFrequency != SF_Pixel && ShaderFrequency != SF_Compute)
-	{
-		return NonPixelShaderExpressionError();
-	}
-	return AddInlinedCodeChunkZeroDeriv(MCT_Float, TEXT("Parameters.TwoSidedSign"));
+	return AddExternalCodeChunk(TEXT("TwoSidedSign"));
 }
 
 int32 FHLSLMaterialTranslator::VertexNormal()
@@ -12155,12 +12151,12 @@ int32 FHLSLMaterialTranslator::AtmosphericFogColor( int32 WorldPosition, EPositi
 
 int32 FHLSLMaterialTranslator::AtmosphericLightVector()
 {
-	return AddCodeChunk(MCT_Float3, TEXT("MaterialExpressionAtmosphericLightVector(Parameters)"));
+	return AddExternalCodeChunk(TEXT("AtmosphericLightVector"));
 }
 
 int32 FHLSLMaterialTranslator::AtmosphericLightColor()
 {
-	return AddCodeChunk(MCT_Float3, TEXT("MaterialExpressionAtmosphericLightColor(Parameters)"));
+	return AddExternalCodeChunk(TEXT("AtmosphericLightColor"));
 }
 
 int32 FHLSLMaterialTranslator::SkyAtmosphereLightIlluminance(int32 WorldPosition, EPositionOrigin PositionOrigin, int32 LightIndex)
diff --git a/Engine/Source/Runtime/Engine/Private/Materials/MaterialExpressions.cpp b/Engine/Source/Runtime/Engine/Private/Materials/MaterialExpressions.cpp
--- a/Engine/Source/Runtime/Engine/Private/Materials/MaterialExpressions.cpp
+++ b/Engine/Source/Runtime/Engine/Private/Materials/MaterialExpressions.cpp
@@ -20195,11 +20195,6 @@ void UMaterialExpressionDistance::GetCaption(TArray<FString>& OutCaptions) const
 #endif // WITH_EDITOR
 
 #if WITH_EDITOR
-int32 UMaterialExpressionTwoSidedSign::Compile(class FMaterialCompiler* Compiler, int32 OutputIndex)
-{
-	return Compiler->TwoSidedSign();
-}
-
 void UMaterialExpressionTwoSidedSign::GetCaption(TArray<FString>& OutCaptions) const
 {
 	OutCaptions.Add(TEXT("TwoSidedSign"));
@@ -21271,12 +21266,6 @@ uint32 UMaterialExpressionVertexInterpolator::GetOutputType(int32 OutputIndex)
 // UMaterialExpressionrAtmosphericLightVector
 ///////////////////////////////////////////////////////////////////////////////
 #if WITH_EDITOR
-int32 UMaterialExpressionAtmosphericLightVector::Compile(class FMaterialCompiler* Compiler, int32 OutputIndex)
-{
-
-	return Compiler->AtmosphericLightVector();
-}
-
 void UMaterialExpressionAtmosphericLightVector::GetCaption(TArray<FString>& OutCaptions) const
 {
 	OutCaptions.Add(TEXT("Atmosphere Sun Light Vector"));
@@ -21287,12 +21276,6 @@ void UMaterialExpressionAtmosphericLightVector::GetCaption(TArray<FString>& OutC
 // UMaterialExpressionrAtmosphericLightColor
 ///////////////////////////////////////////////////////////////////////////////
 #if WITH_EDITOR
-int32 UMaterialExpressionAtmosphericLightColor::Compile(class FMaterialCompiler* Compiler, int32 OutputIndex)
-{
-
-	return Compiler->AtmosphericLightColor();
-}
-
 void UMaterialExpressionAtmosphericLightColor::GetCaption(TArray<FString>& OutCaptions) const
 {
 	OutCaptions.Add(TEXT("Atmosphere Sun Light Illuminance On Ground (not per pixel)"));
diff --git a/Engine/Source/Runtime/Engine/Public/Materials/MaterialExpressionAtmosphericLightColor.h b/Engine/Source/Runtime/Engine/Public/Materials/MaterialExpressionAtmosphericLightColor.h
--- a/Engine/Source/Runtime/Engine/Public/Materials/MaterialExpressionAtmosphericLightColor.h
+++ b/Engine/Source/Runtime/Engine/Public/Materials/MaterialExpressionAtmosphericLightColor.h
@@ -5,11 +5,11 @@
 
 #include "CoreMinimal.h"
 #include "UObject/ObjectMacros.h"
-#include "Materials/MaterialExpression.h"
+#include "Materials/MaterialExpressionExternalCodeBase.h"
 #include "MaterialExpressionAtmosphericLightColor.generated.h"
 
 UCLASS(collapsecategories, hidecategories = Object, DisplayName = "Atmosphere Sun Light Illuminance On Ground")
-class UMaterialExpressionAtmosphericLightColor : public UMaterialExpression
+class UMaterialExpressionAtmosphericLightColor : public UMaterialExpressionExternalCodeBase
 {
 	GENERATED_BODY()
 
@@ -17,7 +17,6 @@ public:
 
 	//~ Begin UMaterialExpression Interface
 #if WITH_EDITOR
-	virtual int32 Compile(class FMaterialCompiler* Compiler, int32 OutputIndex) override;
 	virtual void GetCaption(TArray<FString>& OutCaptions) const override;
 
 #endif
diff --git a/Engine/Source/Runtime/Engine/Public/Materials/MaterialExpressionAtmosphericLightVector.h b/Engine/Source/Runtime/Engine/Public/Materials/MaterialExpressionAtmosphericLightVector.h
--- a/Engine/Source/Runtime/Engine/Public/Materials/MaterialExpressionAtmosphericLightVector.h
+++ b/Engine/Source/Runtime/Engine/Public/Materials/MaterialExpressionAtmosphericLightVector.h
@@ -5,11 +5,11 @@
 
 #include "CoreMinimal.h"
 #include "UObject/ObjectMacros.h"
-#include "Materials/MaterialExpression.h"
+#include "Materials/MaterialExpressionExternalCodeBase.h"
 #include "MaterialExpressionAtmosphericLightVector.generated.h"
 
 UCLASS(collapsecategories, hidecategories = Object, DisplayName = "Atmosphere Sun Light Vector")
-class UMaterialExpressionAtmosphericLightVector : public UMaterialExpression
+class UMaterialExpressionAtmosphericLightVector : public UMaterialExpressionExternalCodeBase
 {
 	GENERATED_BODY()
 
@@ -17,7 +17,6 @@ public:
 
 	//~ Begin UMaterialExpression Interface
 #if WITH_EDITOR
-	virtual int32 Compile(class FMaterialCompiler* Compiler, int32 OutputIndex) override;
 	virtual void GetCaption(TArray<FString>& OutCaptions) const override;
 
 #endif
diff --git a/Engine/Source/Runtime/Engine/Public/Materials/MaterialExpressionTwoSidedSign.h b/Engine/Source/Runtime/Engine/Public/Materials/MaterialExpressionTwoSidedSign.h
--- a/Engine/Source/Runtime/Engine/Public/Materials/MaterialExpressionTwoSidedSign.h
+++ b/Engine/Source/Runtime/Engine/Public/Materials/MaterialExpressionTwoSidedSign.h
@@ -5,11 +5,11 @@
 
 #include "CoreMinimal.h"
 #include "UObject/ObjectMacros.h"
-#include "Materials/MaterialExpression.h"
+#include "Materials/MaterialExpressionExternalCodeBase.h"
 #include "MaterialExpressionTwoSidedSign.generated.h"
 
 UCLASS(collapsecategories, hidecategories=Object, MinimalAPI)
-class UMaterialExpressionTwoSidedSign : public UMaterialExpression
+class UMaterialExpressionTwoSidedSign : public UMaterialExpressionExternalCodeBase
 {
 	GENERATED_BODY()
 
@@ -17,7 +17,6 @@ public:
 
 	//~ Begin UMaterialExpression Interface
 #if WITH_EDITOR
-	virtual int32 Compile(class FMaterialCompiler* Compiler, int32 OutputIndex) override;
 	virtual void GetCaption(TArray<FString>& OutCaptions) const override;
 #endif
 	//~ End UMaterialExpression Interface
