Add device profiles for Adreno 8xx devices

diff --git a/Engine/Config/BaseDeviceProfiles.ini b/Engine/Config/BaseDeviceProfiles.ini
--- a/Engine/Config/BaseDeviceProfiles.ini
+++ b/Engine/Config/BaseDeviceProfiles.ini
@@ -113,6 +113,8 @@
 +DeviceProfileNameAndTypes=Android_Adreno6xx_Vulkan,Android
 +DeviceProfileNameAndTypes=Android_Adreno7xx,Android
 +DeviceProfileNameAndTypes=Android_Adreno7xx_Vulkan,Android
++DeviceProfileNameAndTypes=Android_Adreno8xx,Android
++DeviceProfileNameAndTypes=Android_Adreno8xx_Vulkan,Android
 +DeviceProfileNameAndTypes=Android_Mali_T6xx,Android
 +DeviceProfileNameAndTypes=Android_Mali_T7xx,Android
 +DeviceProfileNameAndTypes=Android_Mali_T8xx,Android
@@ -774,8 +776,8 @@ BaseProfileName=VisionPro
 +MatchProfile=(Profile="Android_Adreno4xx",Match=((SourceType=SRC_GpuFamily,CompareType=CMP_Regex,MatchString="Adreno \\(TM\\) 4[0-9][0-9]")))
 +MatchProfile=(Profile="Android_Adreno5xx_Low",Match=((SourceType=SRC_GpuFamily,CompareType=CMP_Regex,MatchString="Adreno \\(TM\\) 5[0-1][0-9]")))
 +MatchProfile=(Profile="Android_Adreno5xx",Match=((SourceType=SRC_GpuFamily,CompareType=CMP_Regex,MatchString="Adreno \\(TM\\) 5[0-9][0-9]")))
-; Enable SM5 on Adreno 7xx when SM5 is enabled
-+MatchProfile=(Profile="Android_Adreno_Vulkan_SM5",Match=((SourceType=SRC_GpuFamily,CompareType=CMP_Regex,MatchString="Adreno \\(TM\\) 7[0-9][0-9]"),(SourceType=SRC_AndroidVersion, CompareType=CMP_Regex,MatchString="([0-9]+).*"),(SourceType=SRC_PreviousRegexMatch,CompareType=CMP_GreaterEqual,MatchString="10"),(SourceType=SRC_SM5Available,CompareType=CMP_Equal,MatchString="true")))
+; Enable SM5 on Adreno 7xx and 8xx when SM5 is enabled
++MatchProfile=(Profile="Android_Adreno_Vulkan_SM5",Match=((SourceType=SRC_GpuFamily,CompareType=CMP_Regex,MatchString="Adreno \\(TM\\) [7-8][0-9][0-9]"),(SourceType=SRC_AndroidVersion, CompareType=CMP_Regex,MatchString="([0-9]+).*"),(SourceType=SRC_PreviousRegexMatch,CompareType=CMP_GreaterEqual,MatchString="10"),(SourceType=SRC_SM5Available,CompareType=CMP_Equal,MatchString="true")))
 ; disable Vulkan on Adreno 660 devices with driver version <= 525
 +MatchProfile=(Profile="Android_Adreno6xx",Match=((SourceType=SRC_GpuFamily,CompareType=CMP_Regex,MatchString="Adreno \\(TM\\) 660"),(SourceType=SRC_GLVersion, CompareType=CMP_Regex,MatchString="V@([0-9]+)"),(SourceType=SRC_PreviousRegexMatch,CompareType=CMP_LessEqual,MatchString="525")))
 ; Enable Vulkan on Adreno 6xx devices with Android >= 10 and driver version >= 400
@@ -784,6 +786,10 @@ BaseProfileName=VisionPro
 ; Enable Vulkan on Adreno 7xx devices with Android >= 10
 +MatchProfile=(Profile="Android_Adreno7xx_Vulkan",Match=((SourceType=SRC_GpuFamily,CompareType=CMP_Regex,MatchString="Adreno \\(TM\\) 7[0-9][0-9]"),(SourceType=SRC_AndroidVersion, CompareType=CMP_Regex,MatchString="([0-9]+).*"),(SourceType=SRC_PreviousRegexMatch,CompareType=CMP_GreaterEqual,MatchString="10"),(SourceType=SRC_VulkanAvailable,CompareType=CMP_Equal,MatchString="true")))
 +MatchProfile=(Profile="Android_Adreno7xx",Match=((SourceType=SRC_GpuFamily,CompareType=CMP_Regex,MatchString="Adreno \\(TM\\) 7[0-9][0-9]")))
+; Enable Vulkan on all Adreno 8xx devices when Vulkan is available
++MatchProfile=(Profile="Android_Adreno8xx_Vulkan",Match=((SourceType=SRC_GpuFamily,CompareType=CMP_Regex,MatchString="Adreno \\(TM\\) 8[0-9][0-9]"),(SourceType=SRC_VulkanAvailable,CompareType=CMP_Equal,MatchString="true")))
++MatchProfile=(Profile="Android_Adreno8xx",Match=((SourceType=SRC_GpuFamily,CompareType=CMP_Regex,MatchString="Adreno \\(TM\\) 8[0-9][0-9]")))
+
 ; Enable Vulkan on any unknown Adreno running Android 10 or later
 +MatchProfile=(Profile="Android_Unknown_Vulkan",Match=((SourceType=SRC_GpuFamily,CompareType=CMP_Regex,MatchString="Adreno"),(SourceType=SRC_AndroidVersion, CompareType=CMP_Regex,MatchString="([0-9]+).*"),(SourceType=SRC_PreviousRegexMatch,CompareType=CMP_GreaterEqual,MatchString="10"),(SourceType=SRC_VulkanAvailable,CompareType=CMP_Equal,MatchString="true")))
 
@@ -1023,6 +1029,16 @@ BaseProfileName=Android_Adreno7xx
 ; enable Vulkan
 +CVars=r.Android.DisableVulkanSupport=0
 
+[Android_Adreno8xx DeviceProfile]
+DeviceType=Android
+BaseProfileName=Android_High
+
+[Android_Adreno8xx_Vulkan DeviceProfile]
+DeviceType=Android
+BaseProfileName=Android_Adreno8xx
+; enable Vulkan
++CVars=r.Android.DisableVulkanSupport=0
+
 ;----------------
 ; Xclipse devices
 
