[Materials] Convert material expressions with multiple output slots to data-driven description.

This changes the INI property "ExternalCodeIdentifier" from FName to TArray<FName> and renames it to "ExternalCodeIdentifiers".
During code generation in the `Build()` function, we iterate over ExternalCodeIdentifiers and emit their output into their respective slot index via the Output() emitter function.

#rnx

diff --git a/Engine/Config/BaseMaterialExpressions.ini b/Engine/Config/BaseMaterialExpressions.ini
--- a/Engine/Config/BaseMaterialExpressions.ini
+++ b/Engine/Config/BaseMaterialExpressions.ini
@@ -96,12 +96,12 @@ bShaderInputData=true
 
 [/Script/Engine.MaterialExpressionAtmosphericLightColor]
 +MenuCategories=NSLOCTEXT("MaterialExpression", "Utility", "Utility")
-ExternalCodeIdentifier="AtmosphericLightColor"
++ExternalCodeIdentifiers="AtmosphericLightColor"
 
 [/Script/Engine.MaterialExpressionAtmosphericLightVector]
 +MenuCategories=NSLOCTEXT("MaterialExpression", "Utility", "Utility")
 bShaderInputData=true
-ExternalCodeIdentifier="AtmosphericLightVector"
++ExternalCodeIdentifiers="AtmosphericLightVector"
 
 [/Script/Engine.MaterialExpressionBindlessSwitch]
 +MenuCategories=NSLOCTEXT("MaterialExpression", "Utility", "Utility")
@@ -152,6 +152,10 @@ bShaderInputData=true
 
 [/Script/Engine.MaterialExpressionCloudSampleAttribute]
 +MenuCategories=NSLOCTEXT("MaterialExpression", "Volume", "Volume")
++ExternalCodeIdentifiers="GetCloudSampleAltitude"
++ExternalCodeIdentifiers="GetCloudSampleAltitudeInLayer"
++ExternalCodeIdentifiers="GetCloudSampleNormAltitudeInLayer"
++ExternalCodeIdentifiers="GetCloudSampleShadowSampleDistance"
 
 [/Script/Engine.MaterialExpressionCollectionParameter]
 +MenuCategories=NSLOCTEXT("MaterialExpression", "Parameters", "Parameters")
@@ -323,7 +327,7 @@ bShaderInputData=true
 
 [/Script/Engine.MaterialExpressionFontSignedDistance]
 +MenuCategories=NSLOCTEXT("MaterialExpression", "Constants", "Constants")
-ExternalCodeIdentifier="FontSignedDistanceData"
++ExternalCodeIdentifiers="FontSignedDistanceData"
 Caption="Font Signed Distance"
 
 [/Script/Engine.MaterialExpressionFrac]
@@ -373,7 +377,7 @@ bShaderInputData=true
 
 [/Script/Engine.MaterialExpressionIsOrthographic]
 +MenuCategories=NSLOCTEXT("MaterialExpression", "Constants", "Constants")
-ExternalCodeIdentifier="IsOrthographic"
++ExternalCodeIdentifiers="IsOrthographic"
 bShaderInputData=true
 
 [/Script/Engine.MaterialExpressionLength]
@@ -388,7 +392,7 @@ bShaderInputData=true
 
 [/Script/Engine.MaterialExpressionLightVector]
 +MenuCategories=NSLOCTEXT("MaterialExpression", "Vectors", "Vectors")
-ExternalCodeIdentifier="LightVector"
++ExternalCodeIdentifiers="LightVector"
 
 [/Script/Engine.MaterialExpressionLinearInterpolate]
 +MenuCategories=NSLOCTEXT("MaterialExpression", "Math", "Math")
@@ -428,11 +432,11 @@ bShaderInputData=true
 
 [/Script/Engine.MaterialExpressionMeshPaintTextureCoordinateIndex]
 +MenuCategories=NSLOCTEXT("MaterialExpression", "MeshPaintTexture", "MeshPaintTexture")
-ExternalCodeIdentifier="MeshPaintTextureCoordinateIndex"
++ExternalCodeIdentifiers="MeshPaintTextureCoordinateIndex"
 
 [/Script/Engine.MaterialExpressionMeshPaintTextureObject]
 +MenuCategories=NSLOCTEXT("MaterialExpression", "MeshPaintTexture", "MeshPaintTexture")
-ExternalCodeIdentifier="MeshPaintTextureDescriptor"
++ExternalCodeIdentifiers="MeshPaintTextureDescriptor"
 
 [/Script/Engine.MaterialExpressionMeshPaintTextureReplace]
 +MenuCategories=NSLOCTEXT("MaterialExpression", "MeshPaintTexture", "MeshPaintTexture")
@@ -509,7 +513,7 @@ bShaderInputData=true
 
 [/Script/Engine.MaterialExpressionParticleMacroUV]
 +MenuCategories=NSLOCTEXT("MaterialExpression", "Particles", "Particles")
-ExternalCodeIdentifier="ParticleMacroUV"
++ExternalCodeIdentifiers="ParticleMacroUV"
 bShaderInputData=true
 
 [/Script/Engine.MaterialExpressionParticleMotionBlurFade]
@@ -527,7 +531,7 @@ bShaderInputData=true
 +MenuCategories=NSLOCTEXT("MaterialExpression", "Particles", "Particles")
 +MenuCategories=NSLOCTEXT("MaterialExpression", "Constants", "Constants")
 bShaderInputData=true
-ExternalCodeIdentifier="ParticleRadius"
++ExternalCodeIdentifiers="ParticleRadius"
 
 [/Script/Engine.MaterialExpressionParticleRandom]
 +MenuCategories=NSLOCTEXT("MaterialExpression", "Particles", "Particles")
@@ -617,12 +621,12 @@ bShaderInputData=true
 
 [/Script/Engine.MaterialExpressionPreSkinnedNormal]
 +MenuCategories=NSLOCTEXT("MaterialExpression", "Vectors", "Vectors")
-ExternalCodeIdentifier="PreSkinnedNormal"
++ExternalCodeIdentifiers="PreSkinnedNormal"
 bShaderInputData=true
 
 [/Script/Engine.MaterialExpressionPreSkinnedPosition]
 +MenuCategories=NSLOCTEXT("MaterialExpression", "Vectors", "Vectors")
-ExternalCodeIdentifier="PreSkinnedPosition"
++ExternalCodeIdentifiers="PreSkinnedPosition"
 bShaderInputData=true
 
 [/Script/Engine.MaterialExpressionPreviousFrameSwitch]
@@ -897,7 +901,7 @@ bShaderInputData=true
 [/Script/Engine.MaterialExpressionTwoSidedSign]
 +MenuCategories=NSLOCTEXT("MaterialExpression", "Constants", "Constants")
 bShaderInputData=true
-ExternalCodeIdentifier="TwoSidedSign"
++ExternalCodeIdentifiers="TwoSidedSign"
 
 [/Script/Engine.MaterialExpressionUserSceneTexture]
 +MenuCategories=NSLOCTEXT("MaterialExpression", "Texture", "Texture")
@@ -939,6 +943,8 @@ bShaderInputData=true
 
 [/Script/Engine.MaterialExpressionVolumetricCloudEmptySpaceSkippingInput]
 +MenuCategories=NSLOCTEXT("MaterialExpression", "Volume", "Volume")
++ExternalCodeIdentifiers="GetCloudEmptySpaceSkippingSphereCenterWorldPosition"
++ExternalCodeIdentifiers="GetCloudEmptySpaceSkippingSphereRadius"
 
 [/Script/Engine.MaterialExpressionVolumetricAdvancedMaterialOutput]
 +MenuCategories=NSLOCTEXT("MaterialExpression", "Volume", "Volume")
@@ -979,3 +985,11 @@ bShaderInputData=true
 +ExternalCodeDeclarations=(ReturnType=Float3, Name="AtmosphericLightVector", Definition="MaterialExpressionAtmosphericLightVector(Parameters)")
 +ExternalCodeDeclarations=(ReturnType=Float3, Name="AtmosphericLightColor", Definition="MaterialExpressionAtmosphericLightColor(Parameters)")
 
++ExternalCodeDeclarations=(ReturnType=Float, Name="GetCloudSampleAltitude", Definition="MaterialExpressionCloudSampleAltitude(Parameters)")
++ExternalCodeDeclarations=(ReturnType=Float, Name="GetCloudSampleAltitudeInLayer", Definition="MaterialExpressionCloudSampleAltitudeInLayer(Parameters)")
++ExternalCodeDeclarations=(ReturnType=Float, Name="GetCloudSampleNormAltitudeInLayer", Definition="MaterialExpressionCloudSampleNormAltitudeInLayer(Parameters)")
++ExternalCodeDeclarations=(ReturnType=Float, Name="GetCloudSampleShadowSampleDistance", Definition="MaterialExpressionVolumeSampleShadowSampleDistance(Parameters)")
++ExternalCodeDeclarations=(ReturnType=Float4, Name="GetVolumeSampleConservativeDensity", Definition="MaterialExpressionVolumeSampleConservativeDensity(Parameters)")
++ExternalCodeDeclarations=(ReturnType=Float3, Name="GetCloudEmptySpaceSkippingSphereCenterWorldPosition", Definition="MaterialExpressionCloudEmptySpaceSkippingSphereCenterWorldPosition(Parameters)")
++ExternalCodeDeclarations=(ReturnType=Float1, Name="GetCloudEmptySpaceSkippingSphereRadius", Definition="MaterialExpressionCloudEmptySpaceSkippingSphereRadius(Parameters)")
+
diff --git a/Engine/Source/Runtime/Engine/Private/Materials/HLSLMaterialTranslator.cpp b/Engine/Source/Runtime/Engine/Private/Materials/HLSLMaterialTranslator.cpp
--- a/Engine/Source/Runtime/Engine/Private/Materials/HLSLMaterialTranslator.cpp
+++ b/Engine/Source/Runtime/Engine/Private/Materials/HLSLMaterialTranslator.cpp
@@ -12284,37 +12284,37 @@ int32 FHLSLMaterialTranslator::SceneDepthWithoutWater(int32 Offset, int32 Viewpo
 
 int32 FHLSLMaterialTranslator::GetCloudSampleAltitude()
 {
-	return AddCodeChunk(MCT_Float, TEXT("MaterialExpressionCloudSampleAltitude(Parameters)"));
+	return AddExternalCodeChunk(TEXT("GetCloudSampleAltitude"));
 }
 
 int32 FHLSLMaterialTranslator::GetCloudSampleAltitudeInLayer()
 {
-	return AddCodeChunk(MCT_Float, TEXT("MaterialExpressionCloudSampleAltitudeInLayer(Parameters)"));
+	return AddExternalCodeChunk(TEXT("GetCloudSampleAltitudeInLayer"));
 }
 
 int32 FHLSLMaterialTranslator::GetCloudSampleNormAltitudeInLayer()
 {
-	return AddCodeChunk(MCT_Float, TEXT("MaterialExpressionCloudSampleNormAltitudeInLayer(Parameters)"));
+	return AddExternalCodeChunk(TEXT("GetCloudSampleNormAltitudeInLayer"));
 }
 
 int32 FHLSLMaterialTranslator::GetCloudSampleShadowSampleDistance()
 {
-	return AddCodeChunk(MCT_Float, TEXT("MaterialExpressionVolumeSampleShadowSampleDistance(Parameters)"));
+	return AddExternalCodeChunk(TEXT("GetCloudSampleShadowSampleDistance"));
 }
 
 int32 FHLSLMaterialTranslator::GetVolumeSampleConservativeDensity()
 {
-	return AddCodeChunk(MCT_Float4, TEXT("MaterialExpressionVolumeSampleConservativeDensity(Parameters)"));
+	return AddExternalCodeChunk(TEXT("GetVolumeSampleConservativeDensity"));
 }
 
 int32 FHLSLMaterialTranslator::GetCloudEmptySpaceSkippingSphereCenterWorldPosition()
 {
-	return AddCodeChunk(MCT_Float3, TEXT("MaterialExpressionCloudEmptySpaceSkippingSphereCenterWorldPosition(Parameters)"));
+	return AddExternalCodeChunk(TEXT("GetCloudEmptySpaceSkippingSphereCenterWorldPosition"));
 }
 
 int32 FHLSLMaterialTranslator::GetCloudEmptySpaceSkippingSphereRadius()
 {
-	return AddCodeChunk(MCT_Float1, TEXT("MaterialExpressionCloudEmptySpaceSkippingSphereRadius(Parameters)"));
+	return AddExternalCodeChunk(TEXT("GetCloudEmptySpaceSkippingSphereRadius"));
 }
 
 int32 FHLSLMaterialTranslator::CustomPrimitiveData(int32 OutputIndex, EMaterialValueType Type)
diff --git a/Engine/Source/Runtime/Engine/Private/Materials/MaterialExpressions.cpp b/Engine/Source/Runtime/Engine/Private/Materials/MaterialExpressions.cpp
--- a/Engine/Source/Runtime/Engine/Private/Materials/MaterialExpressions.cpp
+++ b/Engine/Source/Runtime/Engine/Private/Materials/MaterialExpressions.cpp
@@ -3808,11 +3808,24 @@ UMaterialExpressionMeshPaintTextureObject::UMaterialExpressionMeshPaintTextureOb
 
 int32 UMaterialExpressionExternalCodeBase::Compile(class FMaterialCompiler* Compiler, int32 OutputIndex)
 {
-	if (const FMaterialExternalCodeDeclaration* ExternalCode = MaterialExternalCodeRegistry::Get().FindExternalCode(ExternalCodeIdentifier))
+	// Validate output index for given number of code identifiers
+	if (ExternalCodeIdentifiers.Num() == 1)
 	{
-		return Compiler->ExternalCode(*ExternalCode);
+		OutputIndex = 0;
 	}
-	return Compiler->Errorf(TEXT("Missing external code declaration for '%s'"), *ExternalCodeIdentifier.ToString());
+	else if (OutputIndex < 0 || OutputIndex >= ExternalCodeIdentifiers.Num())
+	{
+		return Compiler->Errorf(TEXT("OutputIndex (%d) out of range, material expression has only %d external code entry/entries"), OutputIndex, ExternalCodeIdentifiers.Num());
+	}
+	
+	// Find identifier in external code registry
+	const FMaterialExternalCodeDeclaration* ExternalCode = MaterialExternalCodeRegistry::Get().FindExternalCode(ExternalCodeIdentifiers[OutputIndex]);
+	if (!ExternalCode)
+	{
+		return Compiler->Errorf(TEXT("Missing external code declaration for '%s'"), *ExternalCodeIdentifiers[OutputIndex].ToString());
+	}
+
+	return Compiler->ExternalCode(*ExternalCode);
 }
 
 #endif
@@ -22203,20 +22216,6 @@ UMaterialExpressionVolumetricCloudEmptySpaceSkippingInput::UMaterialExpressionVo
 }
 
 #if WITH_EDITOR
-int32 UMaterialExpressionVolumetricCloudEmptySpaceSkippingInput::Compile(class FMaterialCompiler* Compiler, int32 OutputIndex)
-{
-	if (OutputIndex == 0)
-	{
-		return Compiler->GetCloudEmptySpaceSkippingSphereCenterWorldPosition();
-	}
-	else if (OutputIndex == 1)
-	{
-		return Compiler->GetCloudEmptySpaceSkippingSphereRadius();
-	}
-
-	return Compiler->Errorf(TEXT("Invalid input parameter"));
-}
-
 void UMaterialExpressionVolumetricCloudEmptySpaceSkippingInput::GetCaption(TArray<FString>& OutCaptions) const
 {
 	OutCaptions.Add(TEXT("Volumetric Cloud Empty Space Skipping Input <Experimental>"));
@@ -22307,28 +22306,6 @@ UMaterialExpressionCloudSampleAttribute::UMaterialExpressionCloudSampleAttribute
 }
 
 #if WITH_EDITOR
-int32 UMaterialExpressionCloudSampleAttribute::Compile(class FMaterialCompiler* Compiler, int32 OutputIndex)
-{
-	if (OutputIndex == 0)
-	{
-		return Compiler->GetCloudSampleAltitude();
-	}
-	else if (OutputIndex == 1)
-	{
-		return Compiler->GetCloudSampleAltitudeInLayer();
-	}
-	else if (OutputIndex == 2)
-	{
-		return Compiler->GetCloudSampleNormAltitudeInLayer();
-	}
-	else if (OutputIndex == 3)
-	{
-		return Compiler->GetCloudSampleShadowSampleDistance();
-	}
-
-	return Compiler->Errorf(TEXT("Invalid input parameter"));
-}
-
 void UMaterialExpressionCloudSampleAttribute::GetCaption(TArray<FString>& OutCaptions) const
 {
 	OutCaptions.Add(TEXT("Cloud Sample Attributes"));
diff --git a/Engine/Source/Runtime/Engine/Private/Materials/MaterialExpressionsIR.cpp b/Engine/Source/Runtime/Engine/Private/Materials/MaterialExpressionsIR.cpp
--- a/Engine/Source/Runtime/Engine/Private/Materials/MaterialExpressionsIR.cpp
+++ b/Engine/Source/Runtime/Engine/Private/Materials/MaterialExpressionsIR.cpp
@@ -593,14 +593,23 @@ void UMaterialExpressionTangent::Build(MIR::FEmitter& Emitter)
 
 void UMaterialExpressionExternalCodeBase::Build(MIR::FEmitter& Emitter)
 {
-	if (const FMaterialExternalCodeDeclaration* ExternalCodeDeclaration = MaterialExternalCodeRegistry::Get().FindExternalCode(ExternalCodeIdentifier))
+	if (ExternalCodeIdentifiers.IsEmpty())
 	{
-		MIR::FValue* Value = Emitter.EmitExternalCode(ExternalCodeDeclaration);
-		Emitter.Output(0, Value);
+		Emitter.Error(TEXT("Material expression is missing external code identifiers"));
+		return;
 	}
-	else
+
+	for (int32 OutputIndex = 0; OutputIndex < ExternalCodeIdentifiers.Num(); ++OutputIndex)
 	{
-		Emitter.Errorf(TEXT("Missing external code declaration for '%s'"), *ExternalCodeIdentifier.ToString());
+		if (const FMaterialExternalCodeDeclaration* ExternalCodeDeclaration = MaterialExternalCodeRegistry::Get().FindExternalCode(ExternalCodeIdentifiers[OutputIndex]))
+		{
+			MIR::FValue* Value = Emitter.EmitExternalCode(ExternalCodeDeclaration);
+			Emitter.Output(OutputIndex, Value);
+		}
+		else
+		{
+			Emitter.Errorf(TEXT("Missing external code declaration for '%s' [OutputIndex=%d]"), *ExternalCodeIdentifiers[OutputIndex].ToString(), OutputIndex);
+		}
 	}
 }
 
diff --git a/Engine/Source/Runtime/Engine/Public/Materials/MaterialExpressionCloudLayer.h b/Engine/Source/Runtime/Engine/Public/Materials/MaterialExpressionCloudLayer.h
--- a/Engine/Source/Runtime/Engine/Public/Materials/MaterialExpressionCloudLayer.h
+++ b/Engine/Source/Runtime/Engine/Public/Materials/MaterialExpressionCloudLayer.h
@@ -4,17 +4,16 @@
 
 #include "CoreMinimal.h"
 #include "UObject/ObjectMacros.h"
-#include "Materials/MaterialExpression.h"
+#include "Materials/MaterialExpressionExternalCodeBase.h"
 #include "MaterialExpressionCloudLayer.generated.h"
 
 UCLASS()
-class UMaterialExpressionCloudSampleAttribute : public UMaterialExpression
+class UMaterialExpressionCloudSampleAttribute : public UMaterialExpressionExternalCodeBase
 {
 	GENERATED_UCLASS_BODY()
 
 	//~ Begin UMaterialExpression Interface
 #if WITH_EDITOR
-	virtual int32 Compile(class FMaterialCompiler* Compiler, int32 OutputIndex) override;
 	virtual void GetCaption(TArray<FString>& OutCaptions) const override;
 	virtual void GetExpressionToolTip(TArray<FString>& OutToolTip) override;
 #endif
diff --git a/Engine/Source/Runtime/Engine/Public/Materials/MaterialExpressionExternalCodeBase.h b/Engine/Source/Runtime/Engine/Public/Materials/MaterialExpressionExternalCodeBase.h
--- a/Engine/Source/Runtime/Engine/Public/Materials/MaterialExpressionExternalCodeBase.h
+++ b/Engine/Source/Runtime/Engine/Public/Materials/MaterialExpressionExternalCodeBase.h
@@ -12,9 +12,12 @@ class UMaterialExpressionExternalCodeBase : public UMaterialExpression
 	GENERATED_BODY()
 
 public:
-	/** Identifier of the external HLSL code this expression inserts into the material shader. */
+	/**
+	 * List of identifiers of the external HLSL code this expression inserts into the material shader.
+	 * If the number of elements is 1, the single element will always be used. Otherwise, OutputIndex selects the respective identifier.
+	 */
 	UPROPERTY(Config)
-	FName ExternalCodeIdentifier;
+	TArray<FName> ExternalCodeIdentifiers;
 
 	//~ Begin UMaterialExpression Interface
 #if WITH_EDITOR
diff --git a/Engine/Source/Runtime/Engine/Public/Materials/MaterialExpressionVolumetricAdvancedMaterialInput.h b/Engine/Source/Runtime/Engine/Public/Materials/MaterialExpressionVolumetricAdvancedMaterialInput.h
--- a/Engine/Source/Runtime/Engine/Public/Materials/MaterialExpressionVolumetricAdvancedMaterialInput.h
+++ b/Engine/Source/Runtime/Engine/Public/Materials/MaterialExpressionVolumetricAdvancedMaterialInput.h
@@ -4,7 +4,7 @@
 
 #include "CoreMinimal.h"
 #include "UObject/ObjectMacros.h"
-#include "Materials/MaterialExpression.h"
+#include "Materials/MaterialExpressionExternalCodeBase.h"
 #include "MaterialExpressionVolumetricAdvancedMaterialInput.generated.h"
 
 UCLASS()
@@ -21,13 +21,12 @@ class UMaterialExpressionVolumetricAdvancedMaterialInput : public UMaterialExpre
 };
 
 UCLASS(Experimental)
-class UMaterialExpressionVolumetricCloudEmptySpaceSkippingInput : public UMaterialExpression
+class UMaterialExpressionVolumetricCloudEmptySpaceSkippingInput : public UMaterialExpressionExternalCodeBase
 {
 	GENERATED_UCLASS_BODY()
 
 	//~ Begin UMaterialExpression Interface
 #if WITH_EDITOR
-	virtual int32 Compile(class FMaterialCompiler* Compiler, int32 OutputIndex) override;
 	virtual void GetCaption(TArray<FString>& OutCaptions) const override;
 	virtual void GetConnectorToolTip(int32 InputIndex, int32 OutputIndex, TArray<FString>& OutToolTip) override;
 	virtual uint32 GetInputType(int32 InputIndex) override;
