Add device profile for iPhone16e

diff --git a/Engine/Config/BaseDeviceProfiles.ini b/Engine/Config/BaseDeviceProfiles.ini
--- a/Engine/Config/BaseDeviceProfiles.ini
+++ b/Engine/Config/BaseDeviceProfiles.ini
@@ -54,6 +54,7 @@
 +DeviceProfileNameAndTypes=iPhone15Pro,IOS
 +DeviceProfileNameAndTypes=iPhone15ProMax,IOS
 +DeviceProfileNameAndTypes=iPhone16,IOS
++DeviceProfileNameAndTypes=iPhone16e,IOS
 +DeviceProfileNameAndTypes=iPhone16Plus,IOS
 +DeviceProfileNameAndTypes=iPhone16Pro,IOS
 +DeviceProfileNameAndTypes=iPhone16ProMax,IOS
@@ -257,6 +258,7 @@ iPhone17,1=iPhone16Pro
 iPhone17,2=iPhone16ProMax
 iPhone17,3=iPhone16
 iPhone17,4=iPhone16Plus
+iPhone17,5=iPhone16e
 iPod9,1=iPodTouch7
 iPad5,[1,2]=iPadMini4
 iPad5,[3,4]=iPadAir2
@@ -540,6 +542,10 @@ BaseProfileName=iPhone13ProMax
 DeviceType=IOS
 BaseProfileName=iPhone15
 
+[iPhone16e DeviceProfile]
+DeviceType=IOS
+BaseProfileName=iPhone15
+
 [iPhone16Plus DeviceProfile]
 DeviceType=IOS
 BaseProfileName=iPhone15Plus
