Enable the AssetRegistryDiscoveryCache by Default again
The two blocking issues have been addressed:
UE-203531 and UE-206631

diff --git a/Engine/Config/BaseEngine.ini b/Engine/Config/BaseEngine.ini
--- a/Engine/Config/BaseEngine.ini
+++ b/Engine/Config/BaseEngine.ini
@@ -2954,7 +2954,7 @@ GatherOnGameThreadOnly=false
 ; by setting AssetRegistryDiscoveryCache=AlwaysWrite, and it can be force-read with no test for invalidation due to modified
 ; directories by commandline -AssetRegistryCacheSkipInvalidate.
 ; Expected values for AssetRegistryDiscoveryCache are Never,Default,AlwaysWrite.
-AssetRegistryDiscoveryCache=Never
+AssetRegistryDiscoveryCache=Default
 
 [AutomationTesting]
 ImportTestPath=Content/EditorAutomation/
