Enabled real time skylight captures in all platforms with resolution overrides on Mobile and Switch

#changelist validated

diff --git a/Engine/Config/BaseScalability.ini b/Engine/Config/BaseScalability.ini
--- a/Engine/Config/BaseScalability.ini
+++ b/Engine/Config/BaseScalability.ini
@@ -263,7 +263,6 @@ r.Shadow.Virtual.SMRT.SamplesPerRayLocal=8
 [GlobalIlluminationQuality@0]
 r.DistanceFieldAO=0
 r.Lumen.DiffuseIndirect.Allow=0
-r.SkyLight.RealTimeReflectionCapture=0
 r.RayTracing.Scene.BuildMode=0
 
 [GlobalIlluminationQuality@1]
@@ -271,7 +270,6 @@ r.RayTracing.Scene.BuildMode=0
 r.DistanceFieldAO=1
 r.AOQuality=1
 r.Lumen.DiffuseIndirect.Allow=0
-r.SkyLight.RealTimeReflectionCapture=0
 r.RayTracing.Scene.BuildMode=0
 
 [GlobalIlluminationQuality@2]
@@ -299,7 +297,6 @@ r.Lumen.TranslucencyVolume.GridPixelSize=64
 r.Lumen.TranslucencyVolume.TraceFromVolume=0
 r.Lumen.TranslucencyVolume.RadianceCache.ProbeResolution=8
 r.Lumen.TranslucencyVolume.RadianceCache.NumProbesToTraceBudget=200
-r.SkyLight.RealTimeReflectionCapture=1
 r.RayTracing.Scene.BuildMode=0
 
 [GlobalIlluminationQuality@3]
@@ -328,7 +325,6 @@ r.Lumen.TranslucencyVolume.TraceFromVolume=1
 r.Lumen.TranslucencyVolume.TracingOctahedronResolution=3
 r.Lumen.TranslucencyVolume.RadianceCache.ProbeResolution=8
 r.Lumen.TranslucencyVolume.RadianceCache.NumProbesToTraceBudget=200
-r.SkyLight.RealTimeReflectionCapture=1
 r.RayTracing.Scene.BuildMode=1
 
 [GlobalIlluminationQuality@Cine]
@@ -357,7 +353,6 @@ r.Lumen.TranslucencyVolume.TraceFromVolume=1
 r.Lumen.TranslucencyVolume.TracingOctahedronResolution=4
 r.Lumen.TranslucencyVolume.RadianceCache.ProbeResolution=16
 r.Lumen.TranslucencyVolume.RadianceCache.NumProbesToTraceBudget=1000
-r.SkyLight.RealTimeReflectionCapture=1
 r.RayTracing.Scene.BuildMode=1
 
 
@@ -665,6 +660,7 @@ r.SkyAtmosphere.SampleCountMax=16.0
 r.SkyAtmosphere.TransmittanceLUT.UseSmallFormat=1
 r.SkyAtmosphere.TransmittanceLUT.SampleCount=10.0
 r.SkyAtmosphere.MultiScatteringLUT.SampleCount=15.0
+r.SkyLight.RealTimeReflectionCapture.TimeSlice.SkyCloudCubeFacePerFrame=2
 fx.Niagara.QualityLevel=0
 r.HeterogeneousVolumes=0
 
@@ -693,6 +689,7 @@ r.SkyAtmosphere.SampleCountMax=32.0
 r.SkyAtmosphere.TransmittanceLUT.UseSmallFormat=0
 r.SkyAtmosphere.TransmittanceLUT.SampleCount=10.0
 r.SkyAtmosphere.MultiScatteringLUT.SampleCount=15.0
+r.SkyLight.RealTimeReflectionCapture.TimeSlice.SkyCloudCubeFacePerFrame=3
 fx.Niagara.QualityLevel=1
 r.HeterogeneousVolumes=0
 
@@ -721,6 +718,7 @@ r.SkyAtmosphere.SampleCountMax=64.0
 r.SkyAtmosphere.TransmittanceLUT.UseSmallFormat=0
 r.SkyAtmosphere.TransmittanceLUT.SampleCount=10.0
 r.SkyAtmosphere.MultiScatteringLUT.SampleCount=15.0
+r.SkyLight.RealTimeReflectionCapture.TimeSlice.SkyCloudCubeFacePerFrame=4
 fx.Niagara.QualityLevel=2
 r.HeterogeneousVolumes.DownsampleFactor=3
 r.HeterogeneousVolumes.MaxStepCount=256
@@ -751,6 +749,7 @@ r.SkyAtmosphere.SampleCountMax=128.0
 r.SkyAtmosphere.TransmittanceLUT.UseSmallFormat=0
 r.SkyAtmosphere.TransmittanceLUT.SampleCount=10.0
 r.SkyAtmosphere.MultiScatteringLUT.SampleCount=15.0
+r.SkyLight.RealTimeReflectionCapture.TimeSlice.SkyCloudCubeFacePerFrame=6
 fx.Niagara.QualityLevel=3
 r.Refraction.OffsetQuality=1
 r.HeterogeneousVolumes.DownsampleFactor=2
@@ -784,6 +783,7 @@ r.SkyAtmosphere.SampleCountMax=256.0
 r.SkyAtmosphere.TransmittanceLUT.UseSmallFormat=0
 r.SkyAtmosphere.TransmittanceLUT.SampleCount=30.0
 r.SkyAtmosphere.MultiScatteringLUT.SampleCount=20.0
+r.SkyLight.RealTimeReflectionCapture.TimeSlice.SkyCloudCubeFacePerFrame=6
 fx.Niagara.QualityLevel=4
 r.Refraction.OffsetQuality=1
 r.HeterogeneousVolumes.ApplyFogInscattering=1
