Removed Timeslice params from base scalability to prevent cook errors

diff --git a/Engine/Config/BaseScalability.ini b/Engine/Config/BaseScalability.ini
--- a/Engine/Config/BaseScalability.ini
+++ b/Engine/Config/BaseScalability.ini
@@ -660,7 +660,6 @@ r.SkyAtmosphere.SampleCountMax=16.0
 r.SkyAtmosphere.TransmittanceLUT.UseSmallFormat=1
 r.SkyAtmosphere.TransmittanceLUT.SampleCount=10.0
 r.SkyAtmosphere.MultiScatteringLUT.SampleCount=15.0
-r.SkyLight.RealTimeReflectionCapture.TimeSlice.SkyCloudCubeFacePerFrame=2
 fx.Niagara.QualityLevel=0
 r.HeterogeneousVolumes=0
 
@@ -689,7 +688,6 @@ r.SkyAtmosphere.SampleCountMax=32.0
 r.SkyAtmosphere.TransmittanceLUT.UseSmallFormat=0
 r.SkyAtmosphere.TransmittanceLUT.SampleCount=10.0
 r.SkyAtmosphere.MultiScatteringLUT.SampleCount=15.0
-r.SkyLight.RealTimeReflectionCapture.TimeSlice.SkyCloudCubeFacePerFrame=3
 fx.Niagara.QualityLevel=1
 r.HeterogeneousVolumes=0
 
@@ -718,7 +716,6 @@ r.SkyAtmosphere.SampleCountMax=64.0
 r.SkyAtmosphere.TransmittanceLUT.UseSmallFormat=0
 r.SkyAtmosphere.TransmittanceLUT.SampleCount=10.0
 r.SkyAtmosphere.MultiScatteringLUT.SampleCount=15.0
-r.SkyLight.RealTimeReflectionCapture.TimeSlice.SkyCloudCubeFacePerFrame=4
 fx.Niagara.QualityLevel=2
 r.HeterogeneousVolumes.DownsampleFactor=3
 r.HeterogeneousVolumes.MaxStepCount=256
@@ -749,7 +746,6 @@ r.SkyAtmosphere.SampleCountMax=128.0
 r.SkyAtmosphere.TransmittanceLUT.UseSmallFormat=0
 r.SkyAtmosphere.TransmittanceLUT.SampleCount=10.0
 r.SkyAtmosphere.MultiScatteringLUT.SampleCount=15.0
-r.SkyLight.RealTimeReflectionCapture.TimeSlice.SkyCloudCubeFacePerFrame=6
 fx.Niagara.QualityLevel=3
 r.Refraction.OffsetQuality=1
 r.HeterogeneousVolumes.DownsampleFactor=2
@@ -783,7 +779,6 @@ r.SkyAtmosphere.SampleCountMax=256.0
 r.SkyAtmosphere.TransmittanceLUT.UseSmallFormat=0
 r.SkyAtmosphere.TransmittanceLUT.SampleCount=30.0
 r.SkyAtmosphere.MultiScatteringLUT.SampleCount=20.0
-r.SkyLight.RealTimeReflectionCapture.TimeSlice.SkyCloudCubeFacePerFrame=6
 fx.Niagara.QualityLevel=4
 r.Refraction.OffsetQuality=1
 r.HeterogeneousVolumes.ApplyFogInscattering=1
