[MIR] Forward environment definitions for data-driven external code expressions.

This implements the code generation for several more material expressions in the new translator, primarily the particle expressions.
External code expressions that require enivornment definitions (some depending on the shader stage) are forwarded to the HLSL code generation and described in the data-driven INI file.

#rnx

diff --git a/Engine/Config/BaseMaterialExpressions.ini b/Engine/Config/BaseMaterialExpressions.ini
--- a/Engine/Config/BaseMaterialExpressions.ini
+++ b/Engine/Config/BaseMaterialExpressions.ini
@@ -507,11 +507,13 @@ bShaderInputData=true
 [/Script/Engine.MaterialExpressionParticleColor]
 +MenuCategories=NSLOCTEXT("MaterialExpression", "Particles", "Particles")
 +MenuCategories=NSLOCTEXT("MaterialExpression", "Constants", "Constants")
++ExternalCodeIdentifiers="ParticleColor"
 bShaderInputData=true
 
 [/Script/Engine.MaterialExpressionParticleDirection]
 +MenuCategories=NSLOCTEXT("MaterialExpression", "Particles", "Particles")
 +MenuCategories=NSLOCTEXT("MaterialExpression", "Constants", "Constants")
++ExternalCodeIdentifiers="ParticleDirection"
 bShaderInputData=true
 
 [/Script/Engine.MaterialExpressionParticleMacroUV]
@@ -522,6 +524,7 @@ bShaderInputData=true
 [/Script/Engine.MaterialExpressionParticleMotionBlurFade]
 +MenuCategories=NSLOCTEXT("MaterialExpression", "Particles", "Particles")
 +MenuCategories=NSLOCTEXT("MaterialExpression", "Constants", "Constants")
++ExternalCodeIdentifiers="ParticleMotionBlurFade"
 bShaderInputData=true
 
 [/Script/Engine.MaterialExpressionParticlePositionWS]
@@ -539,26 +542,31 @@ bShaderInputData=true
 [/Script/Engine.MaterialExpressionParticleRandom]
 +MenuCategories=NSLOCTEXT("MaterialExpression", "Particles", "Particles")
 +MenuCategories=NSLOCTEXT("MaterialExpression", "Constants", "Constants")
++ExternalCodeIdentifiers="ParticleRandom"
 bShaderInputData=true
 
 [/Script/Engine.MaterialExpressionParticleRelativeTime]
 +MenuCategories=NSLOCTEXT("MaterialExpression", "Particles", "Particles")
 +MenuCategories=NSLOCTEXT("MaterialExpression", "Constants", "Constants")
++ExternalCodeIdentifiers="ParticleRelativeTime"
 bShaderInputData=true
 
 [/Script/Engine.MaterialExpressionParticleSize]
 +MenuCategories=NSLOCTEXT("MaterialExpression", "Particles", "Particles")
 +MenuCategories=NSLOCTEXT("MaterialExpression", "Constants", "Constants")
++ExternalCodeIdentifiers="ParticleSize"
 bShaderInputData=true
 
 [/Script/Engine.MaterialExpressionParticleSpeed]
 +MenuCategories=NSLOCTEXT("MaterialExpression", "Particles", "Particles")
 +MenuCategories=NSLOCTEXT("MaterialExpression", "Constants", "Constants")
++ExternalCodeIdentifiers="ParticleSpeed"
 bShaderInputData=true
 
 [/Script/Engine.MaterialExpressionParticleSpriteRotation]
 +MenuCategories=NSLOCTEXT("MaterialExpression", "Particles", "Particles")
 +MenuCategories=NSLOCTEXT("MaterialExpression", "Constants", "Constants")
++ExternalCodeIdentifiers="ParticleSpriteRotation"
 bShaderInputData=true
 
 [/Script/Engine.MaterialExpressionParticleSubUV]
@@ -964,38 +972,38 @@ bShaderInputData=true
 +MenuCategories=NSLOCTEXT("MaterialExpression", "Parameters", "Parameters")
 
 [/Script/Engine.MaterialExternalCodeCollection]
-+ExternalCodeDeclarations=(ReturnType=Float3, Name="CameraVector", Definition="Parameters.CameraVector", ShaderFrequency=Pixel|Compute, bIsInlined=true)
-+ExternalCodeDeclarations=(ReturnType=Float4, Name="DecalColor", Definition="DecalColor()", Derivative=Zero, ShaderFrequency=Pixel, Domains=(MD_DeferredDecal))
-+ExternalCodeDeclarations=(ReturnType=Float, Name="DecalLifetimeOpacity", Definition="DecalLifetimeOpacity()", Derivative=Zero, ShaderFrequency=Pixel, Domains=(MD_DeferredDecal))
-+ExternalCodeDeclarations=(ReturnType=Float, Name="DeltaTime", Definition="View.DeltaTime", Derivative=Zero, bIsInlined=true)
-+ExternalCodeDeclarations=(ReturnType=Float, Name="EyeAdaptation", Definition="EyeAdaptationLookup()", Derivative=Zero, ShaderFrequency=Pixel, bIsInlined=true)
-+ExternalCodeDeclarations=(ReturnType=Float4, Name="FontSignedDistanceData", Definition="GetFontSignedDistanceData(Parameters)", ShaderFrequency=Pixel, Domains=(MD_UI))
-+ExternalCodeDeclarations=(ReturnType=Float, Name="IsOrthographic", Definition="IsOrthoProjectionFloat()", Derivative=Zero, bIsInlined=true)
-+ExternalCodeDeclarations=(ReturnType=Float3, Name="LightVector", Definition="Parameters.LightVector", ShaderFrequency=Pixel|Compute, bIsInlined=true)
-+ExternalCodeDeclarations=(ReturnType=Float, Name="MeshPaintTextureCoordinateIndex", Definition="GetMeshPaintTextureCoordinateIndex(GetPrimitiveData(Parameters))", bIsInlined=true)
-+ExternalCodeDeclarations=(ReturnType=TextureMeshPaint, Name="MeshPaintTextureDescriptor", Definition="GetMeshPaintTextureDescriptor(GetPrimitiveData(Parameters))", Derivative=Zero, bIsInlined=true)
-+ExternalCodeDeclarations=(ReturnType=Float4, Name="ParticleColor", Definition="Parameters.Particle.Color", Derivative=Zero, bIsInlined=true)
-+ExternalCodeDeclarations=(ReturnType=Float3, Name="ParticleDirection", Definition="Parameters.Particle.Velocity.xyz", Derivative=Zero, bIsInlined=true)
-+ExternalCodeDeclarations=(ReturnType=Float2, Name="ParticleMacroUV", Definition="GetParticleMacroUV(Parameters)", ShaderFrequency=Pixel|Compute)
-+ExternalCodeDeclarations=(ReturnType=Float, Name="ParticleMotionBlurFade", Definition="Parameters.Particle.MotionBlurFade", Derivative=Zero, bIsInlined=true)
-+ExternalCodeDeclarations=(ReturnType=Float, Name="ParticleRadius", Definition="max(Parameters.Particle.TranslatedWorldPositionAndSize.w, .001f)", Derivative=Zero, bIsInlined=true)
-+ExternalCodeDeclarations=(ReturnType=Float, Name="ParticleRelativeTime", Definition="Parameters.Particle.RelativeTime", Derivative=Zero, bIsInlined=true)
-+ExternalCodeDeclarations=(ReturnType=Float, Name="ParticleRandom", Definition="Parameters.Particle.Random", Derivative=Zero, bIsInlined=true)
-+ExternalCodeDeclarations=(ReturnType=Float, Name="ParticleSpeed", Definition="Parameters.Particle.Velocity.w", Derivative=Zero, bIsInlined=true)
-+ExternalCodeDeclarations=(ReturnType=Float2, Name="ParticleSize", Definition="Parameters.Particle.Size", Derivative=Zero, bIsInlined=true)
-+ExternalCodeDeclarations=(ReturnType=Float2, Name="ParticleSpriteRotation", Definition="float2(Parameters.Particle.SpriteRotation, Parameters.Particle.SpriteRotation * 57.2957795131f)", Derivative=Zero, bIsInlined=true)
-+ExternalCodeDeclarations=(ReturnType=Float3, Name="PreSkinnedPosition", Definition="Parameters.PreSkinnedPosition", ShaderFrequency=Vertex, bIsInlined=true)
-+ExternalCodeDeclarations=(ReturnType=Float3, Name="PreSkinnedNormal", Definition="Parameters.PreSkinnedNormal", ShaderFrequency=Vertex, bIsInlined=true)
-+ExternalCodeDeclarations=(ReturnType=Float3, Name="ReflectionVector", Definition="Parameters.ReflectionVector", ShaderFrequency=Pixel|Compute, bIsInlined=true)
-+ExternalCodeDeclarations=(ReturnType=Float, Name="TwoSidedSign", Definition="Parameters.TwoSidedSign", ShaderFrequency=Pixel|Compute, Derivative=Zero, bIsInlined=true)
-+ExternalCodeDeclarations=(ReturnType=Float3, Name="AtmosphericLightVector", Definition="MaterialExpressionAtmosphericLightVector(Parameters)")
-+ExternalCodeDeclarations=(ReturnType=Float3, Name="AtmosphericLightColor", Definition="MaterialExpressionAtmosphericLightColor(Parameters)")
-
-+ExternalCodeDeclarations=(ReturnType=Float, Name="GetCloudSampleAltitude", Definition="MaterialExpressionCloudSampleAltitude(Parameters)")
-+ExternalCodeDeclarations=(ReturnType=Float, Name="GetCloudSampleAltitudeInLayer", Definition="MaterialExpressionCloudSampleAltitudeInLayer(Parameters)")
-+ExternalCodeDeclarations=(ReturnType=Float, Name="GetCloudSampleNormAltitudeInLayer", Definition="MaterialExpressionCloudSampleNormAltitudeInLayer(Parameters)")
-+ExternalCodeDeclarations=(ReturnType=Float, Name="GetCloudSampleShadowSampleDistance", Definition="MaterialExpressionVolumeSampleShadowSampleDistance(Parameters)")
-+ExternalCodeDeclarations=(ReturnType=Float4, Name="GetVolumeSampleConservativeDensity", Definition="MaterialExpressionVolumeSampleConservativeDensity(Parameters)")
-+ExternalCodeDeclarations=(ReturnType=Float3, Name="GetCloudEmptySpaceSkippingSphereCenterWorldPosition", Definition="MaterialExpressionCloudEmptySpaceSkippingSphereCenterWorldPosition(Parameters)")
-+ExternalCodeDeclarations=(ReturnType=Float1, Name="GetCloudEmptySpaceSkippingSphereRadius", Definition="MaterialExpressionCloudEmptySpaceSkippingSphereRadius(Parameters)")
++ExternalCodeDeclarations=(ReturnType=Float3,			Name="CameraVector",					Definition="Parameters.CameraVector", ShaderFrequency=Pixel|Compute, bIsInlined=true, EnvironmentDefines=((Name="USES_TRANSFORM_VECTOR")))
++ExternalCodeDeclarations=(ReturnType=Float4,			Name="DecalColor",						Definition="DecalColor()", Derivative=Zero, ShaderFrequency=Pixel, Domains=(MD_DeferredDecal))
++ExternalCodeDeclarations=(ReturnType=Float,			Name="DecalLifetimeOpacity",			Definition="DecalLifetimeOpacity()", Derivative=Zero, ShaderFrequency=Pixel, Domains=(MD_DeferredDecal))
++ExternalCodeDeclarations=(ReturnType=Float,			Name="DeltaTime",						Definition="View.DeltaTime", Derivative=Zero, bIsInlined=true)
++ExternalCodeDeclarations=(ReturnType=Float,			Name="EyeAdaptation",					Definition="EyeAdaptationLookup()", Derivative=Zero, ShaderFrequency=Pixel, bIsInlined=true)
++ExternalCodeDeclarations=(ReturnType=Float4,			Name="FontSignedDistanceData",			Definition="GetFontSignedDistanceData(Parameters)", ShaderFrequency=Pixel, Domains=(MD_UI))
++ExternalCodeDeclarations=(ReturnType=Float,			Name="IsOrthographic",					Definition="IsOrthoProjectionFloat()", Derivative=Zero, bIsInlined=true)
++ExternalCodeDeclarations=(ReturnType=Float3,			Name="LightVector",						Definition="Parameters.LightVector", ShaderFrequency=Pixel|Compute, bIsInlined=true)
++ExternalCodeDeclarations=(ReturnType=Float,			Name="MeshPaintTextureCoordinateIndex",	Definition="GetMeshPaintTextureCoordinateIndex(GetPrimitiveData(Parameters))", bIsInlined=true)
++ExternalCodeDeclarations=(ReturnType=TextureMeshPaint,	Name="MeshPaintTextureDescriptor",		Definition="GetMeshPaintTextureDescriptor(GetPrimitiveData(Parameters))", Derivative=Zero, bIsInlined=true)
++ExternalCodeDeclarations=(ReturnType=Float4,			Name="ParticleColor",					Definition="Parameters.Particle.Color", Derivative=Zero, bIsInlined=true, EnvironmentDefines=((Name="NEEDS_PARTICLE_COLOR", ShaderFrequency=Pixel|Compute)))
++ExternalCodeDeclarations=(ReturnType=Float3,			Name="ParticleDirection",				Definition="Parameters.Particle.Velocity.xyz", Derivative=Zero, bIsInlined=true, EnvironmentDefines=((Name="NEEDS_PARTICLE_VELOCITY")))
++ExternalCodeDeclarations=(ReturnType=Float2,			Name="ParticleMacroUV",					Definition="GetParticleMacroUV(Parameters)", ShaderFrequency=Pixel|Compute)
++ExternalCodeDeclarations=(ReturnType=Float,			Name="ParticleMotionBlurFade",			Definition="Parameters.Particle.MotionBlurFade", Derivative=Zero, bIsInlined=true, EnvironmentDefines=((Name="USES_PARTICLE_MOTION_BLUR")))
++ExternalCodeDeclarations=(ReturnType=Float,			Name="ParticleRadius",					Definition="max(Parameters.Particle.TranslatedWorldPositionAndSize.w, .001f)", Derivative=Zero, bIsInlined=true, EnvironmentDefines=((Name="NEEDS_PARTICLE_POSITION")))
++ExternalCodeDeclarations=(ReturnType=Float,			Name="ParticleRelativeTime",			Definition="Parameters.Particle.RelativeTime", Derivative=Zero, bIsInlined=true, EnvironmentDefines=((Name="NEEDS_PARTICLE_TIME")))
++ExternalCodeDeclarations=(ReturnType=Float,			Name="ParticleRandom",					Definition="Parameters.Particle.Random", Derivative=Zero, bIsInlined=true, EnvironmentDefines=((Name="NEEDS_PARTICLE_RANDOM")))
++ExternalCodeDeclarations=(ReturnType=Float,			Name="ParticleSpeed",					Definition="Parameters.Particle.Velocity.w", Derivative=Zero, bIsInlined=true, EnvironmentDefines=((Name="NEEDS_PARTICLE_VELOCITY")))
++ExternalCodeDeclarations=(ReturnType=Float2,			Name="ParticleSize",					Definition="Parameters.Particle.Size", Derivative=Zero, bIsInlined=true, EnvironmentDefines=((Name="NEEDS_PARTICLE_SIZE")))
++ExternalCodeDeclarations=(ReturnType=Float2,			Name="ParticleSpriteRotation",			Definition="float2(Parameters.Particle.SpriteRotation, Parameters.Particle.SpriteRotation * 57.2957795131f)", Derivative=Zero, bIsInlined=true, EnvironmentDefines=((Name="NEEDS_PARTICLE_SPRITE_ROTATION")))
++ExternalCodeDeclarations=(ReturnType=Float3,			Name="PreSkinnedPosition",				Definition="Parameters.PreSkinnedPosition", ShaderFrequency=Vertex, bIsInlined=true)
++ExternalCodeDeclarations=(ReturnType=Float3,			Name="PreSkinnedNormal",				Definition="Parameters.PreSkinnedNormal", ShaderFrequency=Vertex, bIsInlined=true)
++ExternalCodeDeclarations=(ReturnType=Float3,			Name="ReflectionVector",				Definition="Parameters.ReflectionVector", ShaderFrequency=Pixel|Compute, bIsInlined=true, EnvironmentDefines=((Name="USES_TRANSFORM_VECTOR")))
++ExternalCodeDeclarations=(ReturnType=Float,			Name="TwoSidedSign",					Definition="Parameters.TwoSidedSign", ShaderFrequency=Pixel|Compute, Derivative=Zero, bIsInlined=true)
++ExternalCodeDeclarations=(ReturnType=Float3,			Name="AtmosphericLightVector",			Definition="MaterialExpressionAtmosphericLightVector(Parameters)")
++ExternalCodeDeclarations=(ReturnType=Float3,			Name="AtmosphericLightColor",			Definition="MaterialExpressionAtmosphericLightColor(Parameters)")
+
++ExternalCodeDeclarations=(ReturnType=Float,	Name="GetCloudSampleAltitude",								Definition="MaterialExpressionCloudSampleAltitude(Parameters)")
++ExternalCodeDeclarations=(ReturnType=Float,	Name="GetCloudSampleAltitudeInLayer",						Definition="MaterialExpressionCloudSampleAltitudeInLayer(Parameters)")
++ExternalCodeDeclarations=(ReturnType=Float,	Name="GetCloudSampleNormAltitudeInLayer",					Definition="MaterialExpressionCloudSampleNormAltitudeInLayer(Parameters)")
++ExternalCodeDeclarations=(ReturnType=Float,	Name="GetCloudSampleShadowSampleDistance",					Definition="MaterialExpressionVolumeSampleShadowSampleDistance(Parameters)")
++ExternalCodeDeclarations=(ReturnType=Float4,	Name="GetVolumeSampleConservativeDensity",					Definition="MaterialExpressionVolumeSampleConservativeDensity(Parameters)")
++ExternalCodeDeclarations=(ReturnType=Float3,	Name="GetCloudEmptySpaceSkippingSphereCenterWorldPosition",	Definition="MaterialExpressionCloudEmptySpaceSkippingSphereCenterWorldPosition(Parameters)")
++ExternalCodeDeclarations=(ReturnType=Float1,	Name="GetCloudEmptySpaceSkippingSphereRadius",				Definition="MaterialExpressionCloudEmptySpaceSkippingSphereRadius(Parameters)")
 
diff --git a/Engine/Source/Runtime/Engine/Private/Materials/MaterialExpressionsIR.cpp b/Engine/Source/Runtime/Engine/Private/Materials/MaterialExpressionsIR.cpp
--- a/Engine/Source/Runtime/Engine/Private/Materials/MaterialExpressionsIR.cpp
+++ b/Engine/Source/Runtime/Engine/Private/Materials/MaterialExpressionsIR.cpp
@@ -1109,12 +1109,12 @@ void UMaterialExpressionViewProperty::Build(MIR::FEmitter& Emitter)
 			Emitter.Output(OutputIndex,
 				Emitter.EmitOperator(MIR::BO_Divide,
 					Emitter.EmitConstantFloatScalar(1.0f),
-					Emitter.EmitInlineHLSL(HLSLCodeType, MoveTemp(HLSLCode), {}, {})));
+					Emitter.EmitInlineHLSL(HLSLCodeType, MoveTemp(HLSLCode), MIR::EValueFlags::SubstituteTagsInInlineHLSL)));
 		}
 		else
 		{
 			// @todo-laura.hermanns - old translator used CastToNonLWCIfDisabled(), but LWC not supported in MIR yet
-			Emitter.Output(OutputIndex, Emitter.EmitInlineHLSL(HLSLCodeType, MoveTemp(HLSLCode), {}, {}));
+			Emitter.Output(OutputIndex, Emitter.EmitInlineHLSL(HLSLCodeType, MoveTemp(HLSLCode), MIR::EValueFlags::SubstituteTagsInInlineHLSL));
 		}
 	}
 }
diff --git a/Engine/Source/Runtime/Engine/Private/Materials/MaterialExternalCodeRegistry.h b/Engine/Source/Runtime/Engine/Private/Materials/MaterialExternalCodeRegistry.h
--- a/Engine/Source/Runtime/Engine/Private/Materials/MaterialExternalCodeRegistry.h
+++ b/Engine/Source/Runtime/Engine/Private/Materials/MaterialExternalCodeRegistry.h
@@ -70,6 +70,24 @@ enum class EMaterialValueTypeBridge : uint32
 	Numeric             = MCT_Numeric UMETA(Hidden),
 };
 
+/**
+ * Structure to enable shader environment defines per external code declaration.
+ * E.g. the external code declaration for "ParticleColor" enables the define "NEEDS_PARTICLE_COLOR" when used in the pixel stage.
+ */
+USTRUCT()
+struct FMaterialExternalCodeEnvironmentDefine
+{
+	GENERATED_BODY();
+
+	/** Name of the environment define to enable. */
+	UPROPERTY(Config)
+	FName Name;
+
+	/** Optional shader frequency to further restrict this environment define besides its code declaration shader frequency. */
+	UPROPERTY(Config)
+	EMaterialShaderFrequency ShaderFrequency = EMaterialShaderFrequency::Any;
+};
+
 /** Declaration of external HLSL code. Such code expressions can be emitted as part of a material translation. */
 USTRUCT()
 struct FMaterialExternalCodeDeclaration
@@ -99,6 +117,10 @@ public:
 	UPROPERTY(Config)
 	TArray<TEnumAsByte<EMaterialDomain>> Domains;
 
+	/** List of shader environment defines to enable for this external code declaration. */
+	UPROPERTY(Config)
+	TArray<FMaterialExternalCodeEnvironmentDefine> EnvironmentDefines;
+
 public:
 	/**
 	 * Function return type as material value type.
diff --git a/Engine/Source/Runtime/Engine/Private/Materials/MaterialIRDebug.cpp b/Engine/Source/Runtime/Engine/Private/Materials/MaterialIRDebug.cpp
--- a/Engine/Source/Runtime/Engine/Private/Materials/MaterialIRDebug.cpp
+++ b/Engine/Source/Runtime/Engine/Private/Materials/MaterialIRDebug.cpp
@@ -308,9 +308,16 @@ struct FDebugDumpIRState
 			Out.Appendf(TEXT(" MipValueMode=\"%s\""), *StaticEnum<ETextureMipValueMode>()->GetDisplayNameTextByValue(TextureSample->MipValueMode).ToString());
 			Out.Appendf(TEXT(" SamplerType=\"%s\""), *StaticEnum<EMaterialSamplerType>()->GetDisplayNameTextByValue(TextureSample->SamplerType).ToString());
 		}
-		else if (const FInlineHLSL* ExternalCode = Instr->As<FInlineHLSL>())
+		else if (const FInlineHLSL* InlineHLSL = Instr->As<FInlineHLSL>())
 		{
-			Out.Appendf(TEXT(" \"%s\""), ExternalCode->Code);
+			if (InlineHLSL->HasFlags(EValueFlags::HasDynamicHLSLCode))
+			{
+				Out.Appendf(TEXT(" \"%s\""), InlineHLSL->Code);
+			}
+			else
+			{
+				Out.Appendf(TEXT(" \"%s\""), *InlineHLSL->ExternalCodeDeclaration->Definition);
+			}
 		}
 	}
 
diff --git a/Engine/Source/Runtime/Engine/Private/Materials/MaterialIREmitter.cpp b/Engine/Source/Runtime/Engine/Private/Materials/MaterialIREmitter.cpp
--- a/Engine/Source/Runtime/Engine/Private/Materials/MaterialIREmitter.cpp
+++ b/Engine/Source/Runtime/Engine/Private/Materials/MaterialIREmitter.cpp
@@ -1860,17 +1860,23 @@ FValue* FEmitter::EmitTextureSample(FValue* Texture, FValue* TexCoord, ESamplerS
 FValue* FEmitter::EmitInlineHLSL(const FType* Type, FString Code, EValueFlags ValueFlags, EGraphProperties UsedGraphProperties)
 {
 	MIR::FInlineHLSL Prototype = MakePrototype<MIR::FInlineHLSL>(Type);
-	Prototype.Flags = ValueFlags;
+	Prototype.Flags = ValueFlags | EValueFlags::HasDynamicHLSLCode;
 	Prototype.GraphProperties = UsedGraphProperties;
 	Prototype.Code = Module->PushUserString(MoveTemp(Code));
+
 	return EmitPrototype(this, Prototype);
 }
 
-FValue* FEmitter::EmitInlineHLSL(const FMaterialExternalCodeDeclaration* InExternalCodeDeclaration)
+FValue* FEmitter::EmitInlineHLSL(const FMaterialExternalCodeDeclaration* InExternalCodeDeclaration, EGraphProperties UsedGraphProperties)
 {
 	check(InExternalCodeDeclaration != nullptr);
 	const FType* ReturnType = FType::FromMaterialValueType(InExternalCodeDeclaration->GetReturnTypeValue());
-	return EmitInlineHLSL(ReturnType, InExternalCodeDeclaration->Definition, {}, {});
+
+	MIR::FInlineHLSL Prototype = MakePrototype<MIR::FInlineHLSL>(ReturnType);
+	Prototype.GraphProperties = UsedGraphProperties;
+	Prototype.ExternalCodeDeclaration = InExternalCodeDeclaration;
+
+	return EmitPrototype(this, Prototype);
 }
 
 const FType* FEmitter::TryGetCommonType(const FType* A, const FType* B)
diff --git a/Engine/Source/Runtime/Engine/Private/Materials/MaterialIRModuleBuilder.cpp b/Engine/Source/Runtime/Engine/Private/Materials/MaterialIRModuleBuilder.cpp
--- a/Engine/Source/Runtime/Engine/Private/Materials/MaterialIRModuleBuilder.cpp
+++ b/Engine/Source/Runtime/Engine/Private/Materials/MaterialIRModuleBuilder.cpp
@@ -172,7 +172,8 @@ struct FMaterialIRModuleBuilderImpl
 		}
 
 		// If on top of the stack there's a different expression, we have a dependency to analyze first.
-		if (CurrContext.ExpressionStack.Last() != Emitter->Expression) {
+		if (CurrContext.ExpressionStack.Last() != Emitter->Expression)
+		{
 			return;
 		}
 
@@ -398,10 +399,15 @@ struct FMaterialIRModuleBuilderImpl
 					ValueAnalyzer.Analyze(Value);
 				}
 
+				ValueAnalyzer.PropagateStateInStage(Value, CurrentStage);
+
 				// Mark the used instruction as analyzed for this stage.
 				Value->Flags |= MIR::EValueFlags(1 << StageIndex);
 			}
 		}
+
+		// Move final environemnt defines from analyzer into output module
+		Module->EnvironmentDefines = MoveTemp(ValueAnalyzer.EnvironmentDefines);
 	}
 
 	void Step_LinkInstructions()
diff --git a/Engine/Source/Runtime/Engine/Private/Materials/MaterialIRToHLSLTranslator.cpp b/Engine/Source/Runtime/Engine/Private/Materials/MaterialIRToHLSLTranslator.cpp
--- a/Engine/Source/Runtime/Engine/Private/Materials/MaterialIRToHLSLTranslator.cpp
+++ b/Engine/Source/Runtime/Engine/Private/Materials/MaterialIRToHLSLTranslator.cpp
@@ -418,22 +418,30 @@ struct FTranslator : FMaterialIRToHLSLTranslation
 
 	void LowerInlineHLSL(const MIR::FInlineHLSL* ExternalCode)
 	{
-		check(ExternalCode->Code != nullptr);
-
-		// Substitute placeholder tokens now unless disabled for custom nodes
-		if (ExternalCode->HasFlags(MIR::EValueFlags::SubstituteTagsInInlineHLSL))
+		if (ExternalCode->HasFlags(MIR::EValueFlags::HasDynamicHLSLCode))
 		{
-			FString Code = ExternalCode->Code;
+			check(ExternalCode->Code != nullptr);
 
-			// @todo-laura.hermanns - Must be replaced as soon as this state is available in MIR::FExternalCode
-			constexpr bool bCompilingPreviousFrame = false;
-			Code.ReplaceInline(TEXT("<PREV>"), bCompilingPreviousFrame ? TEXT("Prev") : TEXT(""));
+			// Substitute placeholder tokens now unless disabled for custom nodes
+			if (ExternalCode->HasFlags(MIR::EValueFlags::SubstituteTagsInInlineHLSL))
+			{
+				FString Code = ExternalCode->Code;
 
-			Printer.Buffer.Append(MoveTemp(Code));
+				// @todo-laura.hermanns - Must be replaced as soon as this state is available in MIR
+				constexpr bool bCompilingPreviousFrame = false;
+				Code.ReplaceInline(TEXT("<PREV>"), bCompilingPreviousFrame ? TEXT("Prev") : TEXT(""));
+	
+				Printer.Buffer.Append(MoveTemp(Code));
+			}
+			else
+			{
+				Printer.Buffer.Append(ExternalCode->Code);
+			}
 		}
 		else
 		{
-			Printer.Buffer.Append(ExternalCode->Code);
+			check(ExternalCode->ExternalCodeDeclaration != nullptr);
+			Printer.Buffer.Append(ExternalCode->ExternalCodeDeclaration->Definition);
 		}
 	}
 
@@ -925,6 +933,11 @@ struct FTranslator : FMaterialIRToHLSLTranslation
 		OutEnvironment.SetDefine(TEXT("MATERIAL_VIRTUALTEXTURE_FEEDBACK"), false);
 		OutEnvironment.SetDefine(TEXT("IS_MATERIAL_SHADER"), true);
 
+		for (const FName& EnvironmentDefine : Module->GetEnvironmentDefines())
+		{
+			OutEnvironment.SetDefine(EnvironmentDefine, true);
+		}
+
 		FMaterialShadingModelField ShadingModels = Material->GetShadingModels();
 		ensure(ShadingModels.IsValid());
 
diff --git a/Engine/Source/Runtime/Engine/Private/Materials/MaterialIRValueAnalyzer.cpp b/Engine/Source/Runtime/Engine/Private/Materials/MaterialIRValueAnalyzer.cpp
--- a/Engine/Source/Runtime/Engine/Private/Materials/MaterialIRValueAnalyzer.cpp
+++ b/Engine/Source/Runtime/Engine/Private/Materials/MaterialIRValueAnalyzer.cpp
@@ -4,6 +4,7 @@
 
 #if WITH_EDITOR
 
+#include "Materials/MaterialExternalCodeRegistry.h"
 #include "Materials/MaterialIRModule.h"
 #include "Materials/MaterialIRDebug.h"
 #include "Materials/MaterialIR.h"
@@ -20,6 +21,32 @@ static void AnalyzeExternalInput(FMaterialIRValueAnalyzer& Analyzer, MIR::FExter
 	 Analyzer.Module->GetStatistics().ExternalInputUsedMask[MIR::Stage_Pixel][(int)ExternalInput->Id] = true;
 }
 
+static EMaterialShaderFrequency MapToMaterialShaderFrequencyOrAny(MIR::EStage Stage)
+{
+	switch (Stage)
+	{
+	case MIR::EStage::Stage_Vertex: return EMaterialShaderFrequency::Vertex;
+	case MIR::EStage::Stage_Pixel: return EMaterialShaderFrequency::Pixel;
+	}
+	return EMaterialShaderFrequency::Any;
+}
+
+static void PropagateStateInStageInlineHLSL(FMaterialIRValueAnalyzer& Analyzer, MIR::FInlineHLSL* InlineHLSL, MIR::EStage Stage)
+{
+	if (!InlineHLSL->HasFlags(MIR::EValueFlags::HasDynamicHLSLCode))
+	{
+		check(InlineHLSL->ExternalCodeDeclaration != nullptr);
+		const EMaterialShaderFrequency StageToMaterialFrequency = MapToMaterialShaderFrequencyOrAny(Stage);
+		for (const FMaterialExternalCodeEnvironmentDefine& EnvironmentDefine : InlineHLSL->ExternalCodeDeclaration->EnvironmentDefines)
+		{
+			if ((int32)(EnvironmentDefine.ShaderFrequency & StageToMaterialFrequency) != 0)
+			{
+				Analyzer.EnvironmentDefines.Add(EnvironmentDefine.Name);
+			}
+		}
+	}
+}
+
 static void AnalyzeTextureObject(FMaterialIRValueAnalyzer& Analyzer, MIR::FTextureObject* TextureObject)
 {
 	EMaterialTextureParameterType ParamType = MIR::Internal::TextureMaterialValueTypeToParameterType(TextureObject->Texture->GetMaterialType());
@@ -197,10 +224,10 @@ void FMaterialIRValueAnalyzer::Analyze(MIR::FValue* Value)
 		}
 	}
 
-	#define EXPAND_CASE(ValueType) \
+	#define EXPAND_CASE(ValueType, ...) \
 	case MIR::VK_##ValueType: \
 	{ \
-		Analyze##ValueType(*this, static_cast<MIR::F##ValueType*>(Value)); \
+		Analyze##ValueType(*this, static_cast<MIR::F##ValueType*>(Value) , ## __VA_ARGS__); \
 		break; \
 	}
 
@@ -215,5 +242,23 @@ void FMaterialIRValueAnalyzer::Analyze(MIR::FValue* Value)
 
 	#undef EXPAND_CASE
 }
+
+void FMaterialIRValueAnalyzer::PropagateStateInStage(MIR::FValue* Value, MIR::EStage Stage)
+{
+	#define EXPAND_CASE(ValueType) \
+	case MIR::VK_##ValueType: \
+	{ \
+		PropagateStateInStage##ValueType(*this, static_cast<MIR::F##ValueType*>(Value), Stage); \
+		break; \
+	}
+
+	switch (Value->Kind)
+	{
+		EXPAND_CASE(InlineHLSL);
+		default: break;
+	}
+
+	#undef EXPAND_CASE
+}
 	
 #endif // #if WITH_EDITOR
diff --git a/Engine/Source/Runtime/Engine/Private/Materials/MaterialIRValueAnalyzer.h b/Engine/Source/Runtime/Engine/Private/Materials/MaterialIRValueAnalyzer.h
--- a/Engine/Source/Runtime/Engine/Private/Materials/MaterialIRValueAnalyzer.h
+++ b/Engine/Source/Runtime/Engine/Private/Materials/MaterialIRValueAnalyzer.h
@@ -21,6 +21,9 @@ struct FMaterialIRValueAnalyzer
 	//
 	void Analyze(MIR::FValue* Value);
 
+	//
+	void PropagateStateInStage(MIR::FValue* Value, MIR::EStage Stage);
+
 	//
 	UMaterial* Material{};
 	
@@ -35,6 +38,9 @@ struct FMaterialIRValueAnalyzer
 		
 	//
 	TArray<uint32, TInlineAllocator<8>> FreeOffsetsPerNumComponents[3];
+
+	// List of enabled shader environment defines.
+	TArray<FName> EnvironmentDefines;
 		
 	//
 	FMaterialInsights* Insights{}; // Optional
diff --git a/Engine/Source/Runtime/Engine/Public/Materials/MaterialExpressionParticleColor.h b/Engine/Source/Runtime/Engine/Public/Materials/MaterialExpressionParticleColor.h
--- a/Engine/Source/Runtime/Engine/Public/Materials/MaterialExpressionParticleColor.h
+++ b/Engine/Source/Runtime/Engine/Public/Materials/MaterialExpressionParticleColor.h
@@ -5,11 +5,11 @@
 
 #include "CoreMinimal.h"
 #include "UObject/ObjectMacros.h"
-#include "Materials/MaterialExpression.h"
+#include "Materials/MaterialExpressionExternalCodeBase.h"
 #include "MaterialExpressionParticleColor.generated.h"
 
 UCLASS(collapsecategories, hidecategories=Object)
-class UMaterialExpressionParticleColor : public UMaterialExpression
+class UMaterialExpressionParticleColor : public UMaterialExpressionExternalCodeBase
 {
 	GENERATED_UCLASS_BODY()
 
diff --git a/Engine/Source/Runtime/Engine/Public/Materials/MaterialExpressionParticleDirection.h b/Engine/Source/Runtime/Engine/Public/Materials/MaterialExpressionParticleDirection.h
--- a/Engine/Source/Runtime/Engine/Public/Materials/MaterialExpressionParticleDirection.h
+++ b/Engine/Source/Runtime/Engine/Public/Materials/MaterialExpressionParticleDirection.h
@@ -9,11 +9,11 @@
 
 #include "CoreMinimal.h"
 #include "UObject/ObjectMacros.h"
-#include "Materials/MaterialExpression.h"
+#include "Materials/MaterialExpressionExternalCodeBase.h"
 #include "MaterialExpressionParticleDirection.generated.h"
 
 UCLASS(collapsecategories, hidecategories=Object)
-class UMaterialExpressionParticleDirection : public UMaterialExpression
+class UMaterialExpressionParticleDirection : public UMaterialExpressionExternalCodeBase
 {
 	GENERATED_BODY()
 
diff --git a/Engine/Source/Runtime/Engine/Public/Materials/MaterialExpressionParticleMotionBlurFade.h b/Engine/Source/Runtime/Engine/Public/Materials/MaterialExpressionParticleMotionBlurFade.h
--- a/Engine/Source/Runtime/Engine/Public/Materials/MaterialExpressionParticleMotionBlurFade.h
+++ b/Engine/Source/Runtime/Engine/Public/Materials/MaterialExpressionParticleMotionBlurFade.h
@@ -9,11 +9,11 @@
 
 #include "CoreMinimal.h"
 #include "UObject/ObjectMacros.h"
-#include "Materials/MaterialExpression.h"
+#include "Materials/MaterialExpressionExternalCodeBase.h"
 #include "MaterialExpressionParticleMotionBlurFade.generated.h"
 
 UCLASS(collapsecategories, hidecategories=Object)
-class UMaterialExpressionParticleMotionBlurFade : public UMaterialExpression
+class UMaterialExpressionParticleMotionBlurFade : public UMaterialExpressionExternalCodeBase
 {
 	GENERATED_BODY()
 
diff --git a/Engine/Source/Runtime/Engine/Public/Materials/MaterialExpressionParticleRandom.h b/Engine/Source/Runtime/Engine/Public/Materials/MaterialExpressionParticleRandom.h
--- a/Engine/Source/Runtime/Engine/Public/Materials/MaterialExpressionParticleRandom.h
+++ b/Engine/Source/Runtime/Engine/Public/Materials/MaterialExpressionParticleRandom.h
@@ -9,11 +9,11 @@
 
 #include "CoreMinimal.h"
 #include "UObject/ObjectMacros.h"
-#include "Materials/MaterialExpression.h"
+#include "Materials/MaterialExpressionExternalCodeBase.h"
 #include "MaterialExpressionParticleRandom.generated.h"
 
 UCLASS(collapsecategories, hidecategories=Object)
-class UMaterialExpressionParticleRandom: public UMaterialExpression
+class UMaterialExpressionParticleRandom : public UMaterialExpressionExternalCodeBase
 {
 	GENERATED_BODY()
 
diff --git a/Engine/Source/Runtime/Engine/Public/Materials/MaterialExpressionParticleRelativeTime.h b/Engine/Source/Runtime/Engine/Public/Materials/MaterialExpressionParticleRelativeTime.h
--- a/Engine/Source/Runtime/Engine/Public/Materials/MaterialExpressionParticleRelativeTime.h
+++ b/Engine/Source/Runtime/Engine/Public/Materials/MaterialExpressionParticleRelativeTime.h
@@ -9,11 +9,11 @@
 
 #include "CoreMinimal.h"
 #include "UObject/ObjectMacros.h"
-#include "Materials/MaterialExpression.h"
+#include "Materials/MaterialExpressionExternalCodeBase.h"
 #include "MaterialExpressionParticleRelativeTime.generated.h"
 
 UCLASS(collapsecategories, hidecategories=Object)
-class UMaterialExpressionParticleRelativeTime : public UMaterialExpression
+class UMaterialExpressionParticleRelativeTime : public UMaterialExpressionExternalCodeBase
 {
 	GENERATED_BODY()
 
diff --git a/Engine/Source/Runtime/Engine/Public/Materials/MaterialExpressionParticleSize.h b/Engine/Source/Runtime/Engine/Public/Materials/MaterialExpressionParticleSize.h
--- a/Engine/Source/Runtime/Engine/Public/Materials/MaterialExpressionParticleSize.h
+++ b/Engine/Source/Runtime/Engine/Public/Materials/MaterialExpressionParticleSize.h
@@ -9,11 +9,11 @@
 
 #include "CoreMinimal.h"
 #include "UObject/ObjectMacros.h"
-#include "Materials/MaterialExpression.h"
+#include "Materials/MaterialExpressionExternalCodeBase.h"
 #include "MaterialExpressionParticleSize.generated.h"
 
 UCLASS(collapsecategories, hidecategories=Object)
-class UMaterialExpressionParticleSize : public UMaterialExpression
+class UMaterialExpressionParticleSize : public UMaterialExpressionExternalCodeBase
 {
 	GENERATED_BODY()
 
diff --git a/Engine/Source/Runtime/Engine/Public/Materials/MaterialExpressionParticleSpeed.h b/Engine/Source/Runtime/Engine/Public/Materials/MaterialExpressionParticleSpeed.h
--- a/Engine/Source/Runtime/Engine/Public/Materials/MaterialExpressionParticleSpeed.h
+++ b/Engine/Source/Runtime/Engine/Public/Materials/MaterialExpressionParticleSpeed.h
@@ -9,11 +9,11 @@
 
 #include "CoreMinimal.h"
 #include "UObject/ObjectMacros.h"
-#include "Materials/MaterialExpression.h"
+#include "Materials/MaterialExpressionExternalCodeBase.h"
 #include "MaterialExpressionParticleSpeed.generated.h"
 
 UCLASS(collapsecategories, hidecategories=Object)
-class UMaterialExpressionParticleSpeed : public UMaterialExpression
+class UMaterialExpressionParticleSpeed : public UMaterialExpressionExternalCodeBase
 {
 	GENERATED_BODY()
 
diff --git a/Engine/Source/Runtime/Engine/Public/Materials/MaterialExpressionParticleSpriteRotation.h b/Engine/Source/Runtime/Engine/Public/Materials/MaterialExpressionParticleSpriteRotation.h
--- a/Engine/Source/Runtime/Engine/Public/Materials/MaterialExpressionParticleSpriteRotation.h
+++ b/Engine/Source/Runtime/Engine/Public/Materials/MaterialExpressionParticleSpriteRotation.h
@@ -9,11 +9,11 @@
 
 #include "CoreMinimal.h"
 #include "UObject/ObjectMacros.h"
-#include "Materials/MaterialExpression.h"
+#include "Materials/MaterialExpressionExternalCodeBase.h"
 #include "MaterialExpressionParticleSpriteRotation.generated.h"
 
 UCLASS(collapsecategories, hidecategories=Object)
-class UMaterialExpressionParticleSpriteRotation : public UMaterialExpression
+class UMaterialExpressionParticleSpriteRotation : public UMaterialExpressionExternalCodeBase
 {
 	GENERATED_UCLASS_BODY()
 
diff --git a/Engine/Source/Runtime/Engine/Public/Materials/MaterialIR.h b/Engine/Source/Runtime/Engine/Public/Materials/MaterialIR.h
--- a/Engine/Source/Runtime/Engine/Public/Materials/MaterialIR.h
+++ b/Engine/Source/Runtime/Engine/Public/Materials/MaterialIR.h
@@ -39,6 +39,9 @@ enum class EValueFlags : uint8
 
 	// Whether to skip substution of abstract tokens like "<PREV>". Only true for custom nodes.
 	SubstituteTagsInInlineHLSL = 1 << 2,
+
+	// Whether 'FInlineHLSL::Code' field is used. Otherwise, 'ExternalCodeDeclaration' is used.
+	HasDynamicHLSLCode = 1 << 3, 
 };
 
 ENUM_CLASS_FLAGS(EValueFlags);
@@ -462,8 +465,11 @@ struct FSetMaterialOutput : TInstruction<VK_SetMaterialOutput>
 // MIR instruction for external code declarations; See BaseMaterialExpressions.ini file.
 struct FInlineHLSL : TInstruction<VK_InlineHLSL>
 {
-	// The inlined HLSL code string.
-	const TCHAR* Code;
+	union
+	{
+		const FMaterialExternalCodeDeclaration* ExternalCodeDeclaration;
+		const TCHAR* Code; // The inlined HLSL code string.
+	};
 };
 
 // Note: if you modify this enum, update the implementations of the
diff --git a/Engine/Source/Runtime/Engine/Public/Materials/MaterialIREmitter.h b/Engine/Source/Runtime/Engine/Public/Materials/MaterialIREmitter.h
--- a/Engine/Source/Runtime/Engine/Public/Materials/MaterialIREmitter.h
+++ b/Engine/Source/Runtime/Engine/Public/Materials/MaterialIREmitter.h
@@ -164,8 +164,8 @@ public:
 	FValue* EmitBranch(FValue* Condition, FValue* True, FValue* False);
 	FValue* EmitConstruct(const FType* Type, FValue* Initializer);
 	FValue* EmitTextureSample(FValue* Texture, FValue* TexCoord, ESamplerSourceMode SamplerSourceMode, ETextureMipValueMode MipValueMode, EMaterialSamplerType SamplerType);
-	FValue* EmitInlineHLSL(const FType* Type, FString Code, EValueFlags ValueFlags, EGraphProperties UsedGraphProperties);
-	FValue* EmitInlineHLSL(const FMaterialExternalCodeDeclaration* InExternalCodeDeclaration);
+	FValue* EmitInlineHLSL(const FType* Type, FString Code, EValueFlags ValueFlags = EValueFlags::None, EGraphProperties UsedGraphProperties = EGraphProperties::None);
+	FValue* EmitInlineHLSL(const FMaterialExternalCodeDeclaration* InExternalCodeDeclaration, EGraphProperties UsedGraphProperties = EGraphProperties::None);
 
 	/* Types */
 
diff --git a/Engine/Source/Runtime/Engine/Public/Materials/MaterialIRModule.h b/Engine/Source/Runtime/Engine/Public/Materials/MaterialIRModule.h
--- a/Engine/Source/Runtime/Engine/Public/Materials/MaterialIRModule.h
+++ b/Engine/Source/Runtime/Engine/Public/Materials/MaterialIRModule.h
@@ -61,6 +61,9 @@ public:
 	// Retrieves the root block (i.e. the "main" scope) for a specific shader stage.
 	const MIR::FBlock& GetRootBlock(MIR::EStage Stage) const { return *RootBlock[Stage]; }
 
+	// Returns a list of all environment define names this module requires to be enabled for shader compilation.
+	const TArray<FName>& GetEnvironmentDefines() const { return EnvironmentDefines; }
+
 	// Provides access to the translated material statistics.
 	const FStatistics& GetStatistics() const { return Statistics; }
 
@@ -116,6 +119,9 @@ private:
 	// Stores user-defined strings.	
 	TArray<FString> UserStrings;
 
+	// Environment define names for shader compilation.
+	TArray<FName> EnvironmentDefines;
+
 	// List of compilation errors.
 	TArray<FError> Errors;
 
