The memory reports now include the information of render targets.

diff --git a/Engine/Config/BaseEngine.ini b/Engine/Config/BaseEngine.ini
--- a/Engine/Config/BaseEngine.ini
+++ b/Engine/Config/BaseEngine.ini
@@ -2304,6 +2304,7 @@ bRecompileOnLoad=true
 +Cmd="ListSpawnedActors"
 +Cmd="wp.DumpDataLayers"
 +Cmd="wp.DumpStreamingSources"
++Cmd="r.DumpRenderTargetPoolMemory"
 
 [MemReportFullCommands]
 ; These commands are run when memreport is executed with `-full` (for backward comaptibility) or `-type=Full` parameter,
