Add VerseClass to our types excluded from incremental cook until it is declaring dependencies correctly

diff --git a/Engine/Config/BaseEditor.ini b/Engine/Config/BaseEditor.ini
--- a/Engine/Config/BaseEditor.ini
+++ b/Engine/Config/BaseEditor.ini
@@ -498,6 +498,7 @@ SaveUnversioned=true
 +IterativeClassScriptPackageAllowList=Allow,<EngineRoot>
 
 +IterativeClassDenyList=/Script/CustomizableObject.CustomizableObject		; UE-224744
++IterativeClassDenyList=/Script/CoreUObject.VerseClass						; Not yet documented
 +IterativeClassDenyList=/Script/Engine.BlueprintGeneratedClass 				; Not yet documented
 +IterativeClassDenyList=/Script/Engine.World 								; Not yet documented
 
