[MIR] Implement per-instance material expressions in new material translator.

This implements PerInstanceFadeAmount, PerInstanceRandom, and EyeAdaptation expressions in the new material translator.
Certain macros are always defined but conditionally enabled depending on whether its requetsed by external code declarations, such as "NEEDS_PARTICLE_COLOR".
This follows the same approach as the old translator where the entire set of environent defines is always set but enabled/disabled depending on internal state of the translator.
In the new translator, this state comes from external code declarations as well as material configurations that are queried via their getter functions such as IsUsedWithInstancedStaticMeshes().

#rnx

diff --git a/Engine/Config/BaseMaterialExpressions.ini b/Engine/Config/BaseMaterialExpressions.ini
--- a/Engine/Config/BaseMaterialExpressions.ini
+++ b/Engine/Config/BaseMaterialExpressions.ini
@@ -297,6 +297,7 @@ bShaderInputData=true
 
 [/Script/Engine.MaterialExpressionEyeAdaptation]
 +MenuCategories=NSLOCTEXT("MaterialExpression", "Utility", "Utility")
++ExternalCodeIdentifiers="EyeAdaptation"
 bShaderInputData=true
 
 [/Script/Engine.MaterialExpressionEyeAdaptationInverse]
@@ -598,10 +599,12 @@ bShaderInputData=true
 
 [/Script/Engine.MaterialExpressionPerInstanceFadeAmount]
 +MenuCategories=NSLOCTEXT("MaterialExpression", "Constants", "Constants")
++ExternalCodeIdentifiers="GetPerInstanceFadeAmount"
 bShaderInputData=true
 
 [/Script/Engine.MaterialExpressionPerInstanceRandom]
 +MenuCategories=NSLOCTEXT("MaterialExpression", "Constants", "Constants")
++ExternalCodeIdentifiers="GetPerInstanceRandom"
 bShaderInputData=true
 
 [/Script/Engine.MaterialExpressionPinBase]
@@ -976,7 +979,7 @@ bShaderInputData=true
 +ExternalCodeDeclarations=(ReturnType=Float4,			Name="DecalColor",						Definition="DecalColor()", Derivative=Zero, ShaderFrequency=Pixel, Domains=(MD_DeferredDecal))
 +ExternalCodeDeclarations=(ReturnType=Float,			Name="DecalLifetimeOpacity",			Definition="DecalLifetimeOpacity()", Derivative=Zero, ShaderFrequency=Pixel, Domains=(MD_DeferredDecal))
 +ExternalCodeDeclarations=(ReturnType=Float,			Name="DeltaTime",						Definition="View.DeltaTime", Derivative=Zero, bIsInlined=true)
-+ExternalCodeDeclarations=(ReturnType=Float,			Name="EyeAdaptation",					Definition="EyeAdaptationLookup()", Derivative=Zero, ShaderFrequency=Pixel, bIsInlined=true)
++ExternalCodeDeclarations=(ReturnType=Float,			Name="EyeAdaptation",					Definition="EyeAdaptationLookup()", Derivative=Zero, ShaderFrequency=Pixel, bIsInlined=true, EnvironmentDefines=((Name="USES_EYE_ADAPTATION")))
 +ExternalCodeDeclarations=(ReturnType=Float4,			Name="FontSignedDistanceData",			Definition="GetFontSignedDistanceData(Parameters)", ShaderFrequency=Pixel, Domains=(MD_UI))
 +ExternalCodeDeclarations=(ReturnType=Float,			Name="IsOrthographic",					Definition="IsOrthoProjectionFloat()", Derivative=Zero, bIsInlined=true)
 +ExternalCodeDeclarations=(ReturnType=Float3,			Name="LightVector",						Definition="Parameters.LightVector", ShaderFrequency=Pixel|Compute, bIsInlined=true)
@@ -999,6 +1002,9 @@ bShaderInputData=true
 +ExternalCodeDeclarations=(ReturnType=Float3,			Name="AtmosphericLightVector",			Definition="MaterialExpressionAtmosphericLightVector(Parameters)")
 +ExternalCodeDeclarations=(ReturnType=Float3,			Name="AtmosphericLightColor",			Definition="MaterialExpressionAtmosphericLightColor(Parameters)")
 
++ExternalCodeDeclarations=(ReturnType=Float,			Name="GetPerInstanceRandom",			Definition="GetPerInstanceRandom(Parameters)", ShaderFrequency=Vertex|Pixel, Derivative=Zero, bIsInlined=true, EnvironmentDefines=((Name="NEEDS_PER_INSTANCE_RANDOM_PS", ShaderFrequency=Pixel)))
++ExternalCodeDeclarations=(ReturnType=Float,			Name="GetPerInstanceFadeAmount",		Definition="GetPerInstanceFadeAmount(Parameters)", ShaderFrequency=Vertex|Pixel, Derivative=Zero, bIsInlined=true, EnvironmentDefines=((Name="USES_PER_INSTANCE_FADE_AMOUNT")))
+
 +ExternalCodeDeclarations=(ReturnType=Float,	Name="GetCloudSampleAltitude",								Definition="MaterialExpressionCloudSampleAltitude(Parameters)")
 +ExternalCodeDeclarations=(ReturnType=Float,	Name="GetCloudSampleAltitudeInLayer",						Definition="MaterialExpressionCloudSampleAltitudeInLayer(Parameters)")
 +ExternalCodeDeclarations=(ReturnType=Float,	Name="GetCloudSampleNormAltitudeInLayer",					Definition="MaterialExpressionCloudSampleNormAltitudeInLayer(Parameters)")
diff --git a/Engine/Source/Runtime/Engine/Private/Materials/MaterialIRModuleBuilder.cpp b/Engine/Source/Runtime/Engine/Private/Materials/MaterialIRModuleBuilder.cpp
--- a/Engine/Source/Runtime/Engine/Private/Materials/MaterialIRModuleBuilder.cpp
+++ b/Engine/Source/Runtime/Engine/Private/Materials/MaterialIRModuleBuilder.cpp
@@ -401,6 +401,22 @@ struct FMaterialIRModuleBuilderImpl
 				Value->Flags |= MIR::EValueFlags(1 << StageIndex);
 			}
 		}
+	}
+
+	void Step_ConsolidateEnvironmentDefines()
+	{
+		// Keep defines if a combined condition is met. Otherwise, remove them from the environemnt defines set.
+		auto KeepDefineConditionally = [this](FName Name, bool bConditionToKeepDefine) -> void
+			{
+				if (!bConditionToKeepDefine)
+				{
+					ValueAnalyzer.EnvironmentDefines.Remove(Name);
+				}
+			};
+
+		KeepDefineConditionally(TEXT("USES_PER_INSTANCE_CUSTOM_DATA"), ValueAnalyzer.Material->bUsedWithInstancedStaticMeshes);
+		KeepDefineConditionally(TEXT("NEEDS_PER_INSTANCE_RANDOM_PS"), ValueAnalyzer.Material->bUsedWithInstancedStaticMeshes);
+		KeepDefineConditionally(TEXT("USES_PER_INSTANCE_FADE_AMOUNT"), ValueAnalyzer.Material->bUsedWithInstancedStaticMeshes);
 
 		// Move final environemnt defines from analyzer into output module
 		Module->EnvironmentDefines = MoveTemp(ValueAnalyzer.EnvironmentDefines);
@@ -558,6 +574,7 @@ bool FMaterialIRModuleBuilder::Build(FMaterialIRModule* TargetModule)
 
 	Impl.Step_FlowValuesIntoMaterialOutputs();
 	Impl.Step_AnalyzeIRGraph();
+	Impl.Step_ConsolidateEnvironmentDefines();
 	Impl.Step_LinkInstructions();
 	Impl.Step_Finalize();
 
diff --git a/Engine/Source/Runtime/Engine/Private/Materials/MaterialIRToHLSLTranslator.cpp b/Engine/Source/Runtime/Engine/Private/Materials/MaterialIRToHLSLTranslator.cpp
--- a/Engine/Source/Runtime/Engine/Private/Materials/MaterialIRToHLSLTranslator.cpp
+++ b/Engine/Source/Runtime/Engine/Private/Materials/MaterialIRToHLSLTranslator.cpp
@@ -916,6 +916,9 @@ struct FTranslator : FMaterialIRToHLSLTranslation
 		OutEnvironment.SetDefine(TEXT("NEEDS_PARTICLE_LOCAL_TO_WORLD"), false);
 		OutEnvironment.SetDefine(TEXT("NEEDS_PARTICLE_WORLD_TO_LOCAL"), false);
 		OutEnvironment.SetDefine(TEXT("NEEDS_PER_INSTANCE_RANDOM_PS"), false);
+		OutEnvironment.SetDefine(TEXT("USES_EYE_ADAPTATION"), false);
+		OutEnvironment.SetDefine(TEXT("USES_PER_INSTANCE_CUSTOM_DATA"), false);
+		OutEnvironment.SetDefine(TEXT("USES_PER_INSTANCE_FADE_AMOUNT"), false);
 		OutEnvironment.SetDefine(TEXT("USES_TRANSFORM_VECTOR"), false);
 		OutEnvironment.SetDefine(TEXT("WANT_PIXEL_DEPTH_OFFSET"), CompilationOutput.bUsesPixelDepthOffset);
 		OutEnvironment.SetDefineAndCompileArgument(TEXT("USES_WORLD_POSITION_OFFSET"), (bool)CompilationOutput.bUsesWorldPositionOffset);
@@ -933,6 +936,8 @@ struct FTranslator : FMaterialIRToHLSLTranslation
 		OutEnvironment.SetDefine(TEXT("MATERIAL_VIRTUALTEXTURE_FEEDBACK"), false);
 		OutEnvironment.SetDefine(TEXT("IS_MATERIAL_SHADER"), true);
 
+		// Set all defines that are defined by the module.
+		// Any conditional exemption via material properties, such as 'Material->IsUsedWithInstancedStaticMeshes()', are handled during the material IR analysis.
 		for (const FName& EnvironmentDefine : Module->GetEnvironmentDefines())
 		{
 			OutEnvironment.SetDefine(EnvironmentDefine, true);
diff --git a/Engine/Source/Runtime/Engine/Private/Materials/MaterialIRValueAnalyzer.h b/Engine/Source/Runtime/Engine/Private/Materials/MaterialIRValueAnalyzer.h
--- a/Engine/Source/Runtime/Engine/Private/Materials/MaterialIRValueAnalyzer.h
+++ b/Engine/Source/Runtime/Engine/Private/Materials/MaterialIRValueAnalyzer.h
@@ -40,7 +40,7 @@ struct FMaterialIRValueAnalyzer
 	TArray<uint32, TInlineAllocator<8>> FreeOffsetsPerNumComponents[3];
 
 	// List of enabled shader environment defines.
-	TArray<FName> EnvironmentDefines;
+	TSet<FName> EnvironmentDefines;
 		
 	//
 	FMaterialInsights* Insights{}; // Optional
diff --git a/Engine/Source/Runtime/Engine/Public/Materials/MaterialExpressionEyeAdaptation.h b/Engine/Source/Runtime/Engine/Public/Materials/MaterialExpressionEyeAdaptation.h
--- a/Engine/Source/Runtime/Engine/Public/Materials/MaterialExpressionEyeAdaptation.h
+++ b/Engine/Source/Runtime/Engine/Public/Materials/MaterialExpressionEyeAdaptation.h
@@ -4,7 +4,7 @@
 
 #include "CoreMinimal.h"
 #include "UObject/ObjectMacros.h"
-#include "Materials/MaterialExpression.h"
+#include "Materials/MaterialExpressionExternalCodeBase.h"
 #include "MaterialExpressionEyeAdaptation.generated.h"
 
 /**
@@ -13,7 +13,7 @@
 
 
 UCLASS(collapsecategories, hidecategories=Object)
-class UMaterialExpressionEyeAdaptation : public UMaterialExpression
+class UMaterialExpressionEyeAdaptation : public UMaterialExpressionExternalCodeBase
 {
 	GENERATED_UCLASS_BODY()
 
diff --git a/Engine/Source/Runtime/Engine/Public/Materials/MaterialExpressionPerInstanceFadeAmount.h b/Engine/Source/Runtime/Engine/Public/Materials/MaterialExpressionPerInstanceFadeAmount.h
--- a/Engine/Source/Runtime/Engine/Public/Materials/MaterialExpressionPerInstanceFadeAmount.h
+++ b/Engine/Source/Runtime/Engine/Public/Materials/MaterialExpressionPerInstanceFadeAmount.h
@@ -5,11 +5,11 @@
 
 #include "CoreMinimal.h"
 #include "UObject/ObjectMacros.h"
-#include "Materials/MaterialExpression.h"
+#include "Materials/MaterialExpressionExternalCodeBase.h"
 #include "MaterialExpressionPerInstanceFadeAmount.generated.h"
 
 UCLASS(collapsecategories, hidecategories=Object)
-class UMaterialExpressionPerInstanceFadeAmount : public UMaterialExpression
+class UMaterialExpressionPerInstanceFadeAmount : public UMaterialExpressionExternalCodeBase
 {
 	GENERATED_BODY()
 
diff --git a/Engine/Source/Runtime/Engine/Public/Materials/MaterialExpressionPerInstanceRandom.h b/Engine/Source/Runtime/Engine/Public/Materials/MaterialExpressionPerInstanceRandom.h
--- a/Engine/Source/Runtime/Engine/Public/Materials/MaterialExpressionPerInstanceRandom.h
+++ b/Engine/Source/Runtime/Engine/Public/Materials/MaterialExpressionPerInstanceRandom.h
@@ -5,11 +5,11 @@
 
 #include "CoreMinimal.h"
 #include "UObject/ObjectMacros.h"
-#include "Materials/MaterialExpression.h"
+#include "Materials/MaterialExpressionExternalCodeBase.h"
 #include "MaterialExpressionPerInstanceRandom.generated.h"
 
 UCLASS(collapsecategories, hidecategories=Object)
-class UMaterialExpressionPerInstanceRandom : public UMaterialExpression
+class UMaterialExpressionPerInstanceRandom : public UMaterialExpressionExternalCodeBase
 {
 	GENERATED_BODY()
 
diff --git a/Engine/Source/Runtime/Engine/Public/Materials/MaterialIRModule.h b/Engine/Source/Runtime/Engine/Public/Materials/MaterialIRModule.h
--- a/Engine/Source/Runtime/Engine/Public/Materials/MaterialIRModule.h
+++ b/Engine/Source/Runtime/Engine/Public/Materials/MaterialIRModule.h
@@ -62,7 +62,7 @@ public:
 	const MIR::FBlock& GetRootBlock(MIR::EStage Stage) const { return *RootBlock[Stage]; }
 
 	// Returns a list of all environment define names this module requires to be enabled for shader compilation.
-	const TArray<FName>& GetEnvironmentDefines() const { return EnvironmentDefines; }
+	const TSet<FName>& GetEnvironmentDefines() const { return EnvironmentDefines; }
 
 	// Provides access to the translated material statistics.
 	const FStatistics& GetStatistics() const { return Statistics; }
@@ -120,7 +120,7 @@ private:
 	TArray<FString> UserStrings;
 
 	// Environment define names for shader compilation.
-	TArray<FName> EnvironmentDefines;
+	TSet<FName> EnvironmentDefines;
 
 	// List of compilation errors.
 	TArray<FError> Errors;
