Disable Python for Cook commandlet.

- Saves about 400MB memory. Also allowing Python to actually be used during the cook can open a can of worms for iterative/incremental cooks.
- Also properly disable Python for CookShaders commandlet.

diff --git a/Engine/Config/BaseEditor.ini b/Engine/Config/BaseEditor.ini
--- a/Engine/Config/BaseEditor.ini
+++ b/Engine/Config/BaseEditor.ini
@@ -26,6 +26,10 @@ GrayscaleFontAtlasSize=2048
 ColorFontAtlasSize=512
 SdfFontAtlasSize=2048
 
+[PythonScriptPlugin]
++DisablePythonForCommandlet=CookShaders
++DisablePythonForCommandlet=Cook
+
 [AssetTools]
 ; Full Package Path being hidden from the Asset view
 ;+DenyListAssetPaths="/Game/__ExternalActors__"
diff --git a/Engine/Config/BaseEngine.ini b/Engine/Config/BaseEngine.ini
--- a/Engine/Config/BaseEngine.ini
+++ b/Engine/Config/BaseEngine.ini
@@ -2880,9 +2880,6 @@ AspectRatioAxisConstraint=AspectRatio_MaintainYFOV
 +Class=ParticleSystem
 +Class=Texture2D
 
-[PythonScriptPlugin]
-+DisablePythonForCommandlet=CookShaders
-
 [AssetRegistry]
 ; See also the AssetRegistry section in BaseEditorPerProjectUserSettings, for user-specific options
 ; Default Path Filters not scanned by the asset registry when searching for all assets
