Add new early 2025 iPads to device profiles.

diff --git a/Engine/Config/BaseDeviceProfiles.ini b/Engine/Config/BaseDeviceProfiles.ini
--- a/Engine/Config/BaseDeviceProfiles.ini
+++ b/Engine/Config/BaseDeviceProfiles.ini
@@ -16,6 +16,9 @@
 +DeviceProfileNameAndTypes=iPadAir4,IOS
 +DeviceProfileNameAndTypes=iPadAir5,IOS
 +DeviceProfileNameAndTypes=iPadAir6,IOS
++DeviceProfileNameAndTypes=iPadAir_13,IOS
++DeviceProfileNameAndTypes=iPadAir7,IOS
++DeviceProfileNameAndTypes=iPadAir7_13,IOS
 +DeviceProfileNameAndTypes=iPadMini4,IOS
 +DeviceProfileNameAndTypes=iPadMini5,IOS
 +DeviceProfileNameAndTypes=iPadMini6,IOS
@@ -295,6 +298,8 @@ iPad14,[3,4]=iPadPro4_11
 iPad14,[5,6]=iPadPro6_129
 iPad14,[8,9]=iPadAir6
 iPad14,1$=iPadAir_13
+iPad15,[3,4]=iPadAir7
+iPad15,[5,6]=iPadAir7_13
 iPad16,[1,2]=iPadMini7
 iPad16,[3,4]=iPadPro5_11
 iPad16,[5,6]=iPadPro7_129
@@ -709,6 +714,19 @@ BaseProfileName=IPadAir4
 DeviceType=IOS
 BaseProfileName=IPadAir4
 
+[iPadAir_13 DeviceProfile]
+DeviceType=IOS
+BaseProfileName=iPadAir5
++CVars=ios.PhysicalScreenDensity=264
+
+[iPadAir7 DeviceProfile]
+DeviceType=IOS
+BaseProfileName=iPadAir6
+
+[iPadAir7_13 DeviceProfile]
+DeviceType=IOS
+BaseProfileName=iPadAir_13
+
 [iPadMini5 DeviceProfile]
 DeviceType=IOS
 BaseProfileName=IPadPro
