[Materials] Implemented more material expressions in MIR.

This implements the following material expressions in the new translator: VertexTangentWS, TextureSampleParameterCube, IsFirstPerson, LightVector, CameraPositionWS, CameraVectorWS, and ViewSize.
Some are implemented purely in the data-driven INI file and otherwise have a custom Build() function.

This also produces an error in the MIR analyzer when an InlineHLSL instruction is used with the wrong material domain.

#rnx

diff --git a/Engine/Config/BaseMaterialExpressions.ini b/Engine/Config/BaseMaterialExpressions.ini
--- a/Engine/Config/BaseMaterialExpressions.ini
+++ b/Engine/Config/BaseMaterialExpressions.ini
@@ -132,6 +132,7 @@ bShaderInputData=true
 
 [/Script/Engine.MaterialExpressionCameraVectorWS]
 +MenuCategories=NSLOCTEXT("MaterialExpression", "Vectors", "Vectors")
++ExternalCodeIdentifiers="CameraVector"
 bShaderInputData=true
 
 [/Script/Engine.MaterialExpressionCeil]
@@ -377,6 +378,7 @@ Caption="Font Signed Distance"
 [/Script/Engine.MaterialExpressionIsFirstPerson]
 +MenuCategories=NSLOCTEXT("MaterialExpression", "Constants", "Constants")
 +MenuCategories=NSLOCTEXT("MaterialExpression", "FirstPerson", "First Person")
++ExternalCodeIdentifiers="IsFirstPerson"
 bShaderInputData=true
 
 [/Script/Engine.MaterialExpressionIsOrthographic]
@@ -939,6 +941,7 @@ bShaderInputData=true
 [/Script/Engine.MaterialExpressionVertexTangentWS]
 +MenuCategories=NSLOCTEXT("MaterialExpression", "Vectors", "Vectors")
 +MenuCategories=NSLOCTEXT("MaterialExpression", "Coordinates", "Coordinates")
++ExternalCodeIdentifiers="VertexTangent"
 bShaderInputData=true
 
 [/Script/Engine.MaterialExpressionViewProperty]
@@ -975,14 +978,17 @@ bShaderInputData=true
 +MenuCategories=NSLOCTEXT("MaterialExpression", "Parameters", "Parameters")
 
 [/Script/Engine.MaterialExternalCodeCollection]
-+ExternalCodeDeclarations=(ReturnType=Float3,			Name="CameraVector",					Definition="Parameters.CameraVector", ShaderFrequency=Pixel|Compute, bIsInlined=true, EnvironmentDefines=((Name="USES_TRANSFORM_VECTOR")))
++ExternalCodeDeclarations=(ReturnType=Float3,			Name="AtmosphericLightVector",			Definition="MaterialExpressionAtmosphericLightVector(Parameters)")
++ExternalCodeDeclarations=(ReturnType=Float3,			Name="AtmosphericLightColor",			Definition="MaterialExpressionAtmosphericLightColor(Parameters)")
++ExternalCodeDeclarations=(ReturnType=Float3,			Name="CameraVector",					Definition="Parameters.CameraVector", ShaderFrequency=Pixel|Compute, bIsInlined=true, EnvironmentDefines=((Name="USES_TRANSFORM_VECTOR", ShaderFrequency=Pixel|Compute)))
 +ExternalCodeDeclarations=(ReturnType=Float4,			Name="DecalColor",						Definition="DecalColor()", Derivative=Zero, ShaderFrequency=Pixel, Domains=(MD_DeferredDecal))
 +ExternalCodeDeclarations=(ReturnType=Float,			Name="DecalLifetimeOpacity",			Definition="DecalLifetimeOpacity()", Derivative=Zero, ShaderFrequency=Pixel, Domains=(MD_DeferredDecal))
 +ExternalCodeDeclarations=(ReturnType=Float,			Name="DeltaTime",						Definition="View.DeltaTime", Derivative=Zero, bIsInlined=true)
 +ExternalCodeDeclarations=(ReturnType=Float,			Name="EyeAdaptation",					Definition="EyeAdaptationLookup()", Derivative=Zero, ShaderFrequency=Pixel, bIsInlined=true, EnvironmentDefines=((Name="USES_EYE_ADAPTATION")))
 +ExternalCodeDeclarations=(ReturnType=Float4,			Name="FontSignedDistanceData",			Definition="GetFontSignedDistanceData(Parameters)", ShaderFrequency=Pixel, Domains=(MD_UI))
++ExternalCodeDeclarations=(ReturnType=Float,			Name="IsFirstPerson",					Definition="(IsFirstPerson(Parameters) ? 1.0f : 0.0f)", Derivative=Zero, bIsInlined=true, Domains=(MD_Surface, MD_Volume))
 +ExternalCodeDeclarations=(ReturnType=Float,			Name="IsOrthographic",					Definition="IsOrthoProjectionFloat()", Derivative=Zero, bIsInlined=true)
-+ExternalCodeDeclarations=(ReturnType=Float3,			Name="LightVector",						Definition="Parameters.LightVector", ShaderFrequency=Pixel|Compute, bIsInlined=true)
++ExternalCodeDeclarations=(ReturnType=Float3,			Name="LightVector",						Definition="Parameters.LightVector", ShaderFrequency=Pixel|Compute, bIsInlined=true, Domains=(MD_LightFunction, MD_DeferredDecal))
 +ExternalCodeDeclarations=(ReturnType=Float,			Name="MeshPaintTextureCoordinateIndex",	Definition="GetMeshPaintTextureCoordinateIndex(GetPrimitiveData(Parameters))", bIsInlined=true)
 +ExternalCodeDeclarations=(ReturnType=TextureMeshPaint,	Name="MeshPaintTextureDescriptor",		Definition="GetMeshPaintTextureDescriptor(GetPrimitiveData(Parameters))", Derivative=Zero, bIsInlined=true)
 +ExternalCodeDeclarations=(ReturnType=Float4,			Name="ParticleColor",					Definition="Parameters.Particle.Color", Derivative=Zero, bIsInlined=true, EnvironmentDefines=((Name="NEEDS_PARTICLE_COLOR", ShaderFrequency=Pixel|Compute)))
@@ -997,10 +1003,9 @@ bShaderInputData=true
 +ExternalCodeDeclarations=(ReturnType=Float2,			Name="ParticleSpriteRotation",			Definition="float2(Parameters.Particle.SpriteRotation, Parameters.Particle.SpriteRotation * 57.2957795131f)", Derivative=Zero, bIsInlined=true, EnvironmentDefines=((Name="NEEDS_PARTICLE_SPRITE_ROTATION")))
 +ExternalCodeDeclarations=(ReturnType=Float3,			Name="PreSkinnedPosition",				Definition="Parameters.PreSkinnedPosition", ShaderFrequency=Vertex, bIsInlined=true)
 +ExternalCodeDeclarations=(ReturnType=Float3,			Name="PreSkinnedNormal",				Definition="Parameters.PreSkinnedNormal", ShaderFrequency=Vertex, bIsInlined=true)
-+ExternalCodeDeclarations=(ReturnType=Float3,			Name="ReflectionVector",				Definition="Parameters.ReflectionVector", ShaderFrequency=Pixel|Compute, bIsInlined=true, EnvironmentDefines=((Name="USES_TRANSFORM_VECTOR")))
++ExternalCodeDeclarations=(ReturnType=Float3,			Name="ReflectionVector",				Definition="Parameters.ReflectionVector", ShaderFrequency=Pixel|Compute, bIsInlined=true, EnvironmentDefines=((Name="USES_TRANSFORM_VECTOR", ShaderFrequency=Pixel|Compute)))
 +ExternalCodeDeclarations=(ReturnType=Float,			Name="TwoSidedSign",					Definition="Parameters.TwoSidedSign", ShaderFrequency=Pixel|Compute, Derivative=Zero, bIsInlined=true)
-+ExternalCodeDeclarations=(ReturnType=Float3,			Name="AtmosphericLightVector",			Definition="MaterialExpressionAtmosphericLightVector(Parameters)")
-+ExternalCodeDeclarations=(ReturnType=Float3,			Name="AtmosphericLightColor",			Definition="MaterialExpressionAtmosphericLightColor(Parameters)")
++ExternalCodeDeclarations=(ReturnType=Float3,			Name="VertexTangent",					Definition="Parameters.TangentToWorld[0]", bIsInlined=true, EnvironmentDefines=((Name="USES_TRANSFORM_VECTOR", ShaderFrequency=Pixel|Compute)))
 
 +ExternalCodeDeclarations=(ReturnType=Float,			Name="GetPerInstanceRandom",			Definition="GetPerInstanceRandom(Parameters)", ShaderFrequency=Vertex|Pixel, Derivative=Zero, bIsInlined=true, EnvironmentDefines=((Name="NEEDS_PER_INSTANCE_RANDOM_PS", ShaderFrequency=Pixel)))
 +ExternalCodeDeclarations=(ReturnType=Float,			Name="GetPerInstanceFadeAmount",		Definition="GetPerInstanceFadeAmount(Parameters)", ShaderFrequency=Vertex|Pixel, Derivative=Zero, bIsInlined=true, EnvironmentDefines=((Name="USES_PER_INSTANCE_FADE_AMOUNT")))
diff --git a/Engine/Source/Runtime/Engine/Internal/Materials/MaterialExpressionUtils.cpp b/Engine/Source/Runtime/Engine/Internal/Materials/MaterialExpressionUtils.cpp
new file mode 100644
--- /dev/null
+++ b/Engine/Source/Runtime/Engine/Internal/Materials/MaterialExpressionUtils.cpp
@@ -0,0 +1,27 @@
+// Copyright Epic Games, Inc. All Rights Reserved.
+
+#include "Materials/MaterialExpressionUtils.h"
+
+#include "Containers/EnumAsByte.h"
+#include "MaterialDomain.h"
+#include "Materials/MaterialExternalCodeRegistry.h"
+
+
+namespace MaterialExpressionUtils
+{
+
+	FString FormatUnsupportedMaterialDomainError(const FMaterialExternalCodeDeclaration& InExternalCode, const FName& AssetPathName)
+	{
+		TStringBuilder<1024> SupportedMaterialDomainsString;
+		for (EMaterialDomain Domain : InExternalCode.Domains)
+		{
+			if (SupportedMaterialDomainsString.Len() > 0)
+			{
+				SupportedMaterialDomainsString.Append(TEXT(", "));
+			}
+			SupportedMaterialDomainsString.Append(*MaterialDomainString(Domain));
+		}
+		return FString::Printf(TEXT("External code '%s' (asset: %s) is only available in the following material domains: %s"), *InExternalCode.Name.ToString(), *AssetPathName.ToString(), *SupportedMaterialDomainsString);
+	}
+
+} // MaterialExpressionUtils
diff --git a/Engine/Source/Runtime/Engine/Internal/Materials/MaterialExpressionUtils.h b/Engine/Source/Runtime/Engine/Internal/Materials/MaterialExpressionUtils.h
new file mode 100644
--- /dev/null
+++ b/Engine/Source/Runtime/Engine/Internal/Materials/MaterialExpressionUtils.h
@@ -0,0 +1,19 @@
+// Copyright Epic Games, Inc. All Rights Reserved.
+
+/**
+ * MaterialExpressionUtils - Utility functions for the material translators
+ */
+
+#pragma once
+
+#include "CoreMinimal.h"
+
+
+struct FMaterialExternalCodeDeclaration;
+
+namespace MaterialExpressionUtils
+{
+
+	FString FormatUnsupportedMaterialDomainError(const FMaterialExternalCodeDeclaration& InExternalCode, const FName& AssetPathName = NAME_None);
+
+} // MaterialExpressionUtils
diff --git a/Engine/Source/Runtime/Engine/Private/Materials/HLSLMaterialTranslator.cpp b/Engine/Source/Runtime/Engine/Private/Materials/HLSLMaterialTranslator.cpp
--- a/Engine/Source/Runtime/Engine/Private/Materials/HLSLMaterialTranslator.cpp
+++ b/Engine/Source/Runtime/Engine/Private/Materials/HLSLMaterialTranslator.cpp
@@ -31,6 +31,7 @@
 #include "Materials/MaterialExpressionSingleLayerWaterMaterialOutput.h"
 #include "Materials/MaterialExpressionTextureBase.h"
 #include "Materials/MaterialExpressionThinTranslucentMaterialOutput.h"
+#include "Materials/MaterialExpressionUtils.h"
 #include "Materials/MaterialExpressionVectorNoise.h"
 #include "Materials/MaterialExpressionViewProperty.h"
 #include "Materials/MaterialExpressionVolumetricAdvancedMaterialOutput.h"
@@ -4156,16 +4157,7 @@ int32 FHLSLMaterialTranslator::AddExternalCodeChunk(const FMaterialExternalCodeD
 		checkf(Material != nullptr, TEXT("Missing material while emitting external HLSL code '%s'"), *InExternalCode.Name.ToString());
 		if (InExternalCode.Domains.Find(Material->GetMaterialDomain()) == INDEX_NONE)
 		{
-			TStringBuilder<1024> SupportedMaterialDomainsString;
-			for (EMaterialDomain Domain : InExternalCode.Domains)
-			{
-				if (SupportedMaterialDomainsString.Len() > 0)
-				{
-					SupportedMaterialDomainsString.Append(TEXT(", "));
-				}
-				SupportedMaterialDomainsString.Append(*MaterialDomainString(Domain));
-			}
-			return Errorf(TEXT("External code '%s' (asset: %s) is only available in the following material domains: %s"), *InExternalCode.Name.ToString(), *Material->GetAssetPath().ToString(), SupportedMaterialDomainsString.ToString());
+			return Error(*MaterialExpressionUtils::FormatUnsupportedMaterialDomainError(InExternalCode, Material->GetAssetPath()));
 		}
 	}
 
@@ -6330,10 +6322,6 @@ int32 FHLSLMaterialTranslator::CameraVector()
 
 int32 FHLSLMaterialTranslator::LightVector()
 {
-	if (!Material->IsLightFunction() && !Material->IsDeferredDecal())
-	{
-		return Errorf(TEXT("LightVector can only be used in LightFunction or DeferredDecal materials"));
-	}
 	return AddExternalCodeChunk(TEXT("LightVector"));
 }
 
@@ -11887,7 +11875,7 @@ int32 FHLSLMaterialTranslator::VertexTangent()
 	{
 		bUsesTransformVector = true;
 	}
-	return AddInlinedCodeChunk(MCT_Float3, TEXT("Parameters.TangentToWorld[0]"));
+	return AddExternalCodeChunk(TEXT("VertexTangent"));
 }
 
 int32 FHLSLMaterialTranslator::PixelNormalWS()
@@ -11914,12 +11902,7 @@ int32 FHLSLMaterialTranslator::PixelNormalWS()
 
 int32 FHLSLMaterialTranslator::IsFirstPerson()
 {
-	if (!CheckPrimitivePropertyCompatibity(ANSI_TO_TCHAR(__FUNCTION__)))
-	{
-		return INDEX_NONE;
-	}
-
-	return AddInlinedCodeChunkZeroDeriv(MCT_Float, TEXT("(IsFirstPerson(Parameters) ? 1.0f : 0.0f)"));
+	return AddExternalCodeChunk(TEXT("IsFirstPerson"));
 }
 
 int32 FHLSLMaterialTranslator::Derivative(int32 A, EDervativeComponent Component)
diff --git a/Engine/Source/Runtime/Engine/Private/Materials/MaterialExpressions.cpp b/Engine/Source/Runtime/Engine/Private/Materials/MaterialExpressions.cpp
--- a/Engine/Source/Runtime/Engine/Private/Materials/MaterialExpressions.cpp
+++ b/Engine/Source/Runtime/Engine/Private/Materials/MaterialExpressions.cpp
@@ -23313,11 +23313,6 @@ FString UMaterialExpressionFirstPersonOutput::GetDisplayName() const
 }
 
 #if WITH_EDITOR
-int32 UMaterialExpressionIsFirstPerson::Compile(class FMaterialCompiler* Compiler, int32 OutputIndex)
-{
-	return Compiler->IsFirstPerson();
-}
-
 void UMaterialExpressionIsFirstPerson::GetCaption(TArray<FString>& OutCaptions) const
 {
 	OutCaptions.Add(TEXT("IsFirstPerson"));
diff --git a/Engine/Source/Runtime/Engine/Private/Materials/MaterialExpressionsIR.cpp b/Engine/Source/Runtime/Engine/Private/Materials/MaterialExpressionsIR.cpp
--- a/Engine/Source/Runtime/Engine/Private/Materials/MaterialExpressionsIR.cpp
+++ b/Engine/Source/Runtime/Engine/Private/Materials/MaterialExpressionsIR.cpp
@@ -213,6 +213,7 @@
 #include "Materials/MaterialExpressionTextureProperty.h"
 #include "Materials/MaterialExpressionTextureSample.h"
 #include "Materials/MaterialExpressionTextureSampleParameter.h"
+#include "Materials/MaterialExpressionTextureSampleParameterCube.h"
 #include "Materials/MaterialExpressionThinTranslucentMaterialOutput.h"
 #include "Materials/MaterialExpressionTime.h"
 #include "Materials/MaterialExpressionTransform.h"
@@ -986,6 +987,17 @@ void UMaterialExpressionTextureSampleParameter::Build(MIR::FEmitter& Emitter)
 	BuildTextureSampleExpression(ParameterValue, this, Emitter);
 }
 
+void UMaterialExpressionTextureSampleParameterCube::Build(MIR::FEmitter& Emitter)
+{
+	if (!Coordinates.GetTracedInput().Expression)
+	{
+		Emitter.Error(TEXT("Cube sample needs UV input"));
+		return;
+	}
+
+	UMaterialExpressionTextureSampleParameter::Build(Emitter);
+}
+
 void UMaterialExpressionTextureObjectParameter::Build(MIR::FEmitter& Emitter)
 {
 	MIR::FValue* ParameterValue = BuildTextureObjectParameter(Emitter, this);
@@ -998,6 +1010,7 @@ void UMaterialExpressionTextureCoordinate::Build(MIR::FEmitter& Emitter)
 	if (UnMirrorU || UnMirrorV)
 	{
 		Emitter.Error(TEXT("Unmirroring unsupported"));
+		return;
 	}
 
 	MIR::FValue* OutputValue = Emitter.EmitExternalInput(MIR::TexCoordIndexToExternalInput(CoordinateIndex));
@@ -1166,35 +1179,48 @@ void UMaterialExpressionConvert::Build(MIR::FEmitter& Emitter)
 	}
 }
 
-void UMaterialExpressionViewProperty::Build(MIR::FEmitter& Emitter)
+static MIR::FValue* BuildViewProperty(MIR::FEmitter& Emitter, EMaterialExposedViewProperty InProperty, bool bInvProperty = false)
 {
-	check(Property < MEVP_MAX);
+	check(InProperty < MEVP_MAX);
 
-	const FMaterialExposedViewPropertyMeta& PropertyMeta = MaterialExternalCodeRegistry::Get().GetExternalViewPropertyCode(Property);
+	const FMaterialExposedViewPropertyMeta& PropertyMeta = MaterialExternalCodeRegistry::Get().GetExternalViewPropertyCode(InProperty);
 
 	const bool bHasCustomInverseCode = PropertyMeta.InvPropertyCode != nullptr;
 
+	FString HLSLCode = bInvProperty && bHasCustomInverseCode ? PropertyMeta.InvPropertyCode : PropertyMeta.PropertyCode;
+	const MIR::FType* HLSLCodeType = MIR::FType::FromMaterialValueType(PropertyMeta.Type);
+
+	if (bInvProperty && !bHasCustomInverseCode)
+	{
+		// Fall back to compute the property's inverse from PropertyCode
+		return Emitter.EmitOperator(MIR::BO_Divide,
+			Emitter.EmitConstantFloatScalar(1.0f),
+			Emitter.EmitInlineHLSL(HLSLCodeType, MoveTemp(HLSLCode), MIR::EValueFlags::SubstituteTagsInInlineHLSL));
+	}
+	else
+	{
+		// @todo-laura.hermanns - old translator used CastToNonLWCIfDisabled(), but LWC not supported in MIR yet
+		return Emitter.EmitInlineHLSL(HLSLCodeType, MoveTemp(HLSLCode), MIR::EValueFlags::SubstituteTagsInInlineHLSL);
+	}
+}
+
+void UMaterialExpressionViewProperty::Build(MIR::FEmitter& Emitter)
+{
 	for (int32 OutputIndex = 0; OutputIndex < 2; ++OutputIndex)
 	{
 		const bool bInvProperty = OutputIndex == 1;
+		Emitter.Output(OutputIndex, BuildViewProperty(Emitter, Property, bInvProperty));
+	}
+}
 
-		FString HLSLCode = bInvProperty && bHasCustomInverseCode ? PropertyMeta.InvPropertyCode : PropertyMeta.PropertyCode;
-		const MIR::FType* HLSLCodeType = MIR::FType::FromMaterialValueType(PropertyMeta.Type);
+void UMaterialExpressionViewSize::Build(MIR::FEmitter& Emitter)
+{
+	Emitter.Output(0, BuildViewProperty(Emitter, MEVP_ViewSize));
+}
 
-		if (bInvProperty && !bHasCustomInverseCode)
-		{
-			// Fall back to compute the property's inverse from PropertyCode
-			Emitter.Output(OutputIndex,
-				Emitter.EmitOperator(MIR::BO_Divide,
-					Emitter.EmitConstantFloatScalar(1.0f),
-					Emitter.EmitInlineHLSL(HLSLCodeType, MoveTemp(HLSLCode), MIR::EValueFlags::SubstituteTagsInInlineHLSL)));
-		}
-		else
-		{
-			// @todo-laura.hermanns - old translator used CastToNonLWCIfDisabled(), but LWC not supported in MIR yet
-			Emitter.Output(OutputIndex, Emitter.EmitInlineHLSL(HLSLCodeType, MoveTemp(HLSLCode), MIR::EValueFlags::SubstituteTagsInInlineHLSL));
-		}
-	}
+void UMaterialExpressionCameraPositionWS::Build(MIR::FEmitter& Emitter)
+{
+	Emitter.Output(0, BuildViewProperty(Emitter, MEVP_WorldSpaceCameraPosition));
 }
 
 void UMaterialExpressionPixelNormalWS::Build(MIR::FEmitter& Emitter)
diff --git a/Engine/Source/Runtime/Engine/Private/Materials/MaterialIRValueAnalyzer.cpp b/Engine/Source/Runtime/Engine/Private/Materials/MaterialIRValueAnalyzer.cpp
--- a/Engine/Source/Runtime/Engine/Private/Materials/MaterialIRValueAnalyzer.cpp
+++ b/Engine/Source/Runtime/Engine/Private/Materials/MaterialIRValueAnalyzer.cpp
@@ -5,6 +5,7 @@
 #if WITH_EDITOR
 
 #include "Materials/MaterialExternalCodeRegistry.h"
+#include "Materials/MaterialExpressionUtils.h"
 #include "Materials/MaterialIRModule.h"
 #include "Materials/MaterialIRDebug.h"
 #include "Materials/MaterialIR.h"
@@ -14,6 +15,7 @@
 #include "Materials/Material.h"
 #include "MaterialShared.h"
 #include "Engine/Texture.h"
+#include "UObject/Package.h"
 
 static void AnalyzeExternalInput(FMaterialIRValueAnalyzer& Analyzer, MIR::FExternalInput* ExternalInput)
 {
@@ -47,6 +49,22 @@ static void PropagateStateInStageInlineHLSL(FMaterialIRValueAnalyzer& Analyzer,
 	}
 }
 
+static void AnalyzeInlineHLSL(FMaterialIRValueAnalyzer& Analyzer, MIR::FInlineHLSL* InlineHLSL)
+{
+	if (!InlineHLSL->HasFlags(MIR::EValueFlags::HasDynamicHLSLCode))
+	{
+		check(InlineHLSL->ExternalCodeDeclaration != nullptr);
+
+		// Validate this external code can be used for the current material domain. Empty list implies no restriction on material domains.
+		if (!InlineHLSL->ExternalCodeDeclaration->Domains.IsEmpty() &&
+			InlineHLSL->ExternalCodeDeclaration->Domains.Find(Analyzer.Material->MaterialDomain) == INDEX_NONE)
+		{
+			FName AssetPathName = Analyzer.Material->GetOutermost()->GetFName();
+			Analyzer.Module->AddError(nullptr, MaterialExpressionUtils::FormatUnsupportedMaterialDomainError(*InlineHLSL->ExternalCodeDeclaration, AssetPathName));
+		}
+	}
+}
+
 static void AnalyzeTextureObject(FMaterialIRValueAnalyzer& Analyzer, MIR::FTextureObject* TextureObject)
 {
 	EMaterialTextureParameterType ParamType = MIR::Internal::TextureMaterialValueTypeToParameterType(TextureObject->Texture->GetMaterialType());
@@ -237,6 +255,7 @@ void FMaterialIRValueAnalyzer::Analyze(MIR::FValue* Value)
 		EXPAND_CASE(TextureObject);
 		EXPAND_CASE(UniformParameter);
 		EXPAND_CASE(SetMaterialOutput);
+		EXPAND_CASE(InlineHLSL);
 		default: break;
 	}
 
diff --git a/Engine/Source/Runtime/Engine/Public/Materials/MaterialExpressionCameraPositionWS.h b/Engine/Source/Runtime/Engine/Public/Materials/MaterialExpressionCameraPositionWS.h
--- a/Engine/Source/Runtime/Engine/Public/Materials/MaterialExpressionCameraPositionWS.h
+++ b/Engine/Source/Runtime/Engine/Public/Materials/MaterialExpressionCameraPositionWS.h
@@ -17,6 +17,7 @@ public:
 
 	//~ Begin UMaterialExpression Interface
 #if WITH_EDITOR
+	virtual void Build(MIR::FEmitter& Emitter) override;
 	virtual int32 Compile(class FMaterialCompiler* Compiler, int32 OutputIndex) override;
 	virtual void GetCaption(TArray<FString>& OutCaptions) const override;
 #endif
diff --git a/Engine/Source/Runtime/Engine/Public/Materials/MaterialExpressionCameraVectorWS.h b/Engine/Source/Runtime/Engine/Public/Materials/MaterialExpressionCameraVectorWS.h
--- a/Engine/Source/Runtime/Engine/Public/Materials/MaterialExpressionCameraVectorWS.h
+++ b/Engine/Source/Runtime/Engine/Public/Materials/MaterialExpressionCameraVectorWS.h
@@ -5,11 +5,11 @@
 
 #include "CoreMinimal.h"
 #include "UObject/ObjectMacros.h"
-#include "Materials/MaterialExpression.h"
+#include "Materials/MaterialExpressionExternalCodeBase.h"
 #include "MaterialExpressionCameraVectorWS.generated.h"
 
 UCLASS(MinimalAPI, collapsecategories, hidecategories=Object)
-class UMaterialExpressionCameraVectorWS : public UMaterialExpression
+class UMaterialExpressionCameraVectorWS : public UMaterialExpressionExternalCodeBase
 {
 	GENERATED_BODY()
 
diff --git a/Engine/Source/Runtime/Engine/Public/Materials/MaterialExpressionIsFirstPerson.h b/Engine/Source/Runtime/Engine/Public/Materials/MaterialExpressionIsFirstPerson.h
--- a/Engine/Source/Runtime/Engine/Public/Materials/MaterialExpressionIsFirstPerson.h
+++ b/Engine/Source/Runtime/Engine/Public/Materials/MaterialExpressionIsFirstPerson.h
@@ -5,11 +5,11 @@
 
 #include "CoreMinimal.h"
 #include "UObject/ObjectMacros.h"
-#include "Materials/MaterialExpression.h"
+#include "Materials/MaterialExpressionExternalCodeBase.h"
 #include "MaterialExpressionIsFirstPerson.generated.h"
 
 UCLASS(collapsecategories, hidecategories = Object, MinimalAPI)
-class UMaterialExpressionIsFirstPerson : public UMaterialExpression
+class UMaterialExpressionIsFirstPerson : public UMaterialExpressionExternalCodeBase
 {
 	GENERATED_BODY()
 
@@ -17,7 +17,6 @@ public:
 
 	//~ Begin UMaterialExpression Interface
 #if WITH_EDITOR
-	virtual int32 Compile(class FMaterialCompiler* Compiler, int32 OutputIndex) override;
 	virtual void GetCaption(TArray<FString>& OutCaptions) const override;
 #endif
 	//~ End UMaterialExpression Interface
diff --git a/Engine/Source/Runtime/Engine/Public/Materials/MaterialExpressionLightVector.h b/Engine/Source/Runtime/Engine/Public/Materials/MaterialExpressionLightVector.h
--- a/Engine/Source/Runtime/Engine/Public/Materials/MaterialExpressionLightVector.h
+++ b/Engine/Source/Runtime/Engine/Public/Materials/MaterialExpressionLightVector.h
@@ -5,11 +5,11 @@
 
 #include "CoreMinimal.h"
 #include "UObject/ObjectMacros.h"
-#include "Materials/MaterialExpression.h"
+#include "Materials/MaterialExpressionExternalCodeBase.h"
 #include "MaterialExpressionLightVector.generated.h"
 
 UCLASS(collapsecategories, hidecategories=Object)
-class UMaterialExpressionLightVector : public UMaterialExpression
+class UMaterialExpressionLightVector : public UMaterialExpressionExternalCodeBase
 {
 	GENERATED_BODY()
 
diff --git a/Engine/Source/Runtime/Engine/Public/Materials/MaterialExpressionTextureSampleParameterCube.h b/Engine/Source/Runtime/Engine/Public/Materials/MaterialExpressionTextureSampleParameterCube.h
--- a/Engine/Source/Runtime/Engine/Public/Materials/MaterialExpressionTextureSampleParameterCube.h
+++ b/Engine/Source/Runtime/Engine/Public/Materials/MaterialExpressionTextureSampleParameterCube.h
@@ -18,6 +18,7 @@ class UMaterialExpressionTextureSampleParameterCube : public UMaterialExpression
 #if WITH_EDITOR
 	//~ Begin UMaterialExpression Interface
 	virtual void GetCaption(TArray<FString>& OutCaptions) const override;
+	virtual void Build(MIR::FEmitter& Emitter) override;
 	virtual int32 Compile(class FMaterialCompiler* Compiler, int32 OutputIndex) override;
 	//~ End UMaterialExpression Interface
 
diff --git a/Engine/Source/Runtime/Engine/Public/Materials/MaterialExpressionVertexTangentWS.h b/Engine/Source/Runtime/Engine/Public/Materials/MaterialExpressionVertexTangentWS.h
--- a/Engine/Source/Runtime/Engine/Public/Materials/MaterialExpressionVertexTangentWS.h
+++ b/Engine/Source/Runtime/Engine/Public/Materials/MaterialExpressionVertexTangentWS.h
@@ -5,11 +5,11 @@
 
 #include "CoreMinimal.h"
 #include "UObject/ObjectMacros.h"
-#include "Materials/MaterialExpression.h"
+#include "Materials/MaterialExpressionExternalCodeBase.h"
 #include "MaterialExpressionVertexTangentWS.generated.h"
 
 UCLASS(MinimalAPI, collapsecategories, hidecategories=Object)
-class UMaterialExpressionVertexTangentWS : public UMaterialExpression
+class UMaterialExpressionVertexTangentWS : public UMaterialExpressionExternalCodeBase
 {
 	GENERATED_BODY()
 
diff --git a/Engine/Source/Runtime/Engine/Public/Materials/MaterialExpressionViewSize.h b/Engine/Source/Runtime/Engine/Public/Materials/MaterialExpressionViewSize.h
--- a/Engine/Source/Runtime/Engine/Public/Materials/MaterialExpressionViewSize.h
+++ b/Engine/Source/Runtime/Engine/Public/Materials/MaterialExpressionViewSize.h
@@ -17,6 +17,7 @@ public:
 
 	//~ Begin UMaterialExpression Interface
 #if WITH_EDITOR
+	virtual void Build(MIR::FEmitter& Emitter) override;
 	virtual int32 Compile(class FMaterialCompiler* Compiler, int32 OutputIndex) override;
 	virtual void GetCaption(TArray<FString>& OutCaptions) const override;
 #endif
