Fixed the reference name from StudioTelemetry.Provider.HordeAnalytics to StudioTelemetry.Provider.Horde so that our C# tolls look for the cirrect analytics backend.

This is commented out in the ini file but is expected to be used as a reference by licensees

diff --git a/Engine/Config/BaseEngine.ini b/Engine/Config/BaseEngine.ini
--- a/Engine/Config/BaseEngine.ini
+++ b/Engine/Config/BaseEngine.ini
@@ -3817,7 +3817,7 @@ SendHardwareData=true
 SendOSData=true
 
 ; Example Log File Analytics
-;[StudioTelemetry.Provider.LogAnalytics] 
+;[StudioTelemetry.Provider.Log] 
 ; Name=LogAnalytics
 ; ProviderModule=AnalyticsLog
 ; UsageType=Editor
@@ -3826,7 +3826,7 @@ SendOSData=true
 ; FolderPath="somefolder"
 
 ; Example Horde Analytics Settings
-;[StudioTelemetry.Provider.HordeAnalytics]
+;[StudioTelemetry.Provider.Horde]
 ; Name=HordeAnalytics
 ; ProviderModule=AnalyticsET
 ; UsageType=EditorAndClient
