[MIR] Implemented parameterization of external code declarations and several more expressions in new material translator.

- Enables parameterization of external code expressions using the '$0' syntax as known from shell scripts, e.g. "MyExternalFunction($0)" will substitute the function argument during IR-to-HLSL translation.
- Implements the following material expressions in MIR: EyeAdaptationInverse, ObjectBounds, ObjectRadius, ParticleSubUVProperties, SphericalParticleOpacity, and Time.
- Interpet "MIR." prefix in environment defines of external code declarations as special case during MIR analysis phase.

#rnx

diff --git a/Engine/Config/BaseMaterialExpressions.ini b/Engine/Config/BaseMaterialExpressions.ini
--- a/Engine/Config/BaseMaterialExpressions.ini
+++ b/Engine/Config/BaseMaterialExpressions.ini
@@ -303,6 +303,7 @@ bShaderInputData=true
 
 [/Script/Engine.MaterialExpressionEyeAdaptationInverse]
 +MenuCategories=NSLOCTEXT("MaterialExpression", "Utility", "Utility")
++ExternalCodeIdentifiers="EyeAdaptationInverse"
 bShaderInputData=true
 
 [/Script/Engine.MaterialExpressionFeatureLevelSwitch]
@@ -481,6 +482,7 @@ bShaderInputData=true
 
 [/Script/Engine.MaterialExpressionObjectBounds]
 +MenuCategories=NSLOCTEXT("MaterialExpression", "Vectors", "Vectors")
++ExternalCodeDeclarations="ObjectBounds"
 bShaderInputData=true
 
 [/Script/Engine.MaterialExpressionObjectLocalBounds]
@@ -499,6 +501,7 @@ bShaderInputData=true
 
 [/Script/Engine.MaterialExpressionObjectRadius]
 +MenuCategories=NSLOCTEXT("MaterialExpression", "Coordinates", "Coordinates")
++ExternalCodeIdentifiers="ObjectRadius"
 bShaderInputData=true
 
 [/Script/Engine.MaterialExpressionOneMinus]
@@ -579,6 +582,9 @@ bShaderInputData=true
 [/Script/Engine.MaterialExpressionParticleSubUVProperties]
 +MenuCategories=NSLOCTEXT("MaterialExpression", "Particles", "Particles")
 +MenuCategories=NSLOCTEXT("MaterialExpression", "Coordinates", "Coordinates")
++ExternalCodeIdentifiers="ParticleSubUVCoords0"
++ExternalCodeIdentifiers="ParticleSubUVCoords1"
++ExternalCodeIdentifiers="ParticleSubUVLerp"
 bShaderInputData=true
 
 [/Script/Engine.MaterialExpressionPathTracingBufferTexture]
@@ -809,6 +815,7 @@ bShaderInputData=true
 
 [/Script/Engine.MaterialExpressionSphericalParticleOpacity]
 +MenuCategories=NSLOCTEXT("MaterialExpression", "Particles", "Particles")
++ExternalCodeIdentifiers="SphericalParticleOpacity"
 
 [/Script/Engine.MaterialExpressionSquareRoot]
 +MenuCategories=NSLOCTEXT("MaterialExpression", "Math", "Math")
@@ -985,12 +992,15 @@ bShaderInputData=true
 +ExternalCodeDeclarations=(ReturnType=Float,			Name="DecalLifetimeOpacity",			Definition="DecalLifetimeOpacity()", Derivative=Zero, ShaderFrequency=Pixel, Domains=(MD_DeferredDecal))
 +ExternalCodeDeclarations=(ReturnType=Float,			Name="DeltaTime",						Definition="View.DeltaTime", Derivative=Zero, bIsInlined=true)
 +ExternalCodeDeclarations=(ReturnType=Float,			Name="EyeAdaptation",					Definition="EyeAdaptationLookup()", Derivative=Zero, ShaderFrequency=Pixel, bIsInlined=true, EnvironmentDefines=((Name="USES_EYE_ADAPTATION")))
++ExternalCodeDeclarations=(ReturnType=Float,			Name="EyeAdaptationInverse",			Definition="EyeAdaptationInverseLookup($0)", ShaderFrequency=Pixel, bIsInlined=true, EnvironmentDefines=((Name="USES_EYE_ADAPTATION")))
 +ExternalCodeDeclarations=(ReturnType=Float4,			Name="FontSignedDistanceData",			Definition="GetFontSignedDistanceData(Parameters)", ShaderFrequency=Pixel, Domains=(MD_UI))
 +ExternalCodeDeclarations=(ReturnType=Float,			Name="IsFirstPerson",					Definition="(IsFirstPerson(Parameters) ? 1.0f : 0.0f)", Derivative=Zero, bIsInlined=true, Domains=(MD_Surface, MD_Volume))
 +ExternalCodeDeclarations=(ReturnType=Float,			Name="IsOrthographic",					Definition="IsOrthoProjectionFloat()", Derivative=Zero, bIsInlined=true)
 +ExternalCodeDeclarations=(ReturnType=Float3,			Name="LightVector",						Definition="Parameters.LightVector", ShaderFrequency=Pixel|Compute, bIsInlined=true, Domains=(MD_LightFunction, MD_DeferredDecal))
 +ExternalCodeDeclarations=(ReturnType=Float,			Name="MeshPaintTextureCoordinateIndex",	Definition="GetMeshPaintTextureCoordinateIndex(GetPrimitiveData(Parameters))", bIsInlined=true)
 +ExternalCodeDeclarations=(ReturnType=TextureMeshPaint,	Name="MeshPaintTextureDescriptor",		Definition="GetMeshPaintTextureDescriptor(GetPrimitiveData(Parameters))", Derivative=Zero, bIsInlined=true)
++ExternalCodeDeclarations=(ReturnType=Float3,			Name="ObjectBounds",					Definition="float3(GetPrimitiveData(Parameters).ObjectBoundsX, GetPrimitiveData(Parameters).ObjectBoundsY, GetPrimitiveData(Parameters).ObjectBoundsZ)", Derivative=Zero, bIsInlined=true, Domains=(MD_Surface, MD_Volume))
++ExternalCodeDeclarations=(ReturnType=Float,			Name="ObjectRadius",					Definition="GetPrimitiveData(Parameters).ObjectRadius", Derivative=Zero, bIsInlined=true, Domains=(MD_Surface, MD_Volume))
 +ExternalCodeDeclarations=(ReturnType=Float4,			Name="ParticleColor",					Definition="Parameters.Particle.Color", Derivative=Zero, bIsInlined=true, EnvironmentDefines=((Name="NEEDS_PARTICLE_COLOR", ShaderFrequency=Pixel|Compute)))
 +ExternalCodeDeclarations=(ReturnType=Float3,			Name="ParticleDirection",				Definition="Parameters.Particle.Velocity.xyz", Derivative=Zero, bIsInlined=true, EnvironmentDefines=((Name="NEEDS_PARTICLE_VELOCITY")))
 +ExternalCodeDeclarations=(ReturnType=Float2,			Name="ParticleMacroUV",					Definition="GetParticleMacroUV(Parameters)", ShaderFrequency=Pixel|Compute)
@@ -1001,9 +1011,13 @@ bShaderInputData=true
 +ExternalCodeDeclarations=(ReturnType=Float,			Name="ParticleSpeed",					Definition="Parameters.Particle.Velocity.w", Derivative=Zero, bIsInlined=true, EnvironmentDefines=((Name="NEEDS_PARTICLE_VELOCITY")))
 +ExternalCodeDeclarations=(ReturnType=Float2,			Name="ParticleSize",					Definition="Parameters.Particle.Size", Derivative=Zero, bIsInlined=true, EnvironmentDefines=((Name="NEEDS_PARTICLE_SIZE")))
 +ExternalCodeDeclarations=(ReturnType=Float2,			Name="ParticleSpriteRotation",			Definition="float2(Parameters.Particle.SpriteRotation, Parameters.Particle.SpriteRotation * 57.2957795131f)", Derivative=Zero, bIsInlined=true, EnvironmentDefines=((Name="NEEDS_PARTICLE_SPRITE_ROTATION")))
++ExternalCodeDeclarations=(ReturnType=Float2,			Name="ParticleSubUVCoords0",			Definition="Parameters.Particle.SubUVCoords[0].xy", EnvironmentDefines=((Name="USE_PARTICLE_SUBUVS")))
++ExternalCodeDeclarations=(ReturnType=Float2,			Name="ParticleSubUVCoords1",			Definition="Parameters.Particle.SubUVCoords[1].xy", EnvironmentDefines=((Name="USE_PARTICLE_SUBUVS")))
++ExternalCodeDeclarations=(ReturnType=Float,			Name="ParticleSubUVLerp",				Definition="Parameters.Particle.SubUVLerp", EnvironmentDefines=((Name="USE_PARTICLE_SUBUVS")))
 +ExternalCodeDeclarations=(ReturnType=Float3,			Name="PreSkinnedPosition",				Definition="Parameters.PreSkinnedPosition", ShaderFrequency=Vertex, bIsInlined=true)
 +ExternalCodeDeclarations=(ReturnType=Float3,			Name="PreSkinnedNormal",				Definition="Parameters.PreSkinnedNormal", ShaderFrequency=Vertex, bIsInlined=true)
 +ExternalCodeDeclarations=(ReturnType=Float3,			Name="ReflectionVector",				Definition="Parameters.ReflectionVector", ShaderFrequency=Pixel|Compute, bIsInlined=true, EnvironmentDefines=((Name="USES_TRANSFORM_VECTOR", ShaderFrequency=Pixel|Compute)))
++ExternalCodeDeclarations=(ReturnType=Float,			Name="SphericalParticleOpacity",		Definition="GetSphericalParticleOpacity(Parameters,$0)", ShaderFrequency=Pixel|Compute, EnvironmentDefines=((Name="NEEDS_PARTICLE_POSITION"), (Name="SPHERICAL_PARTICLE_OPACITY"), (Name="NEEDS_WORLD_POSITION_EXCLUDING_SHADER_OFFSETS"), (Name="MIR.SceneDepth")))
 +ExternalCodeDeclarations=(ReturnType=Float,			Name="TwoSidedSign",					Definition="Parameters.TwoSidedSign", ShaderFrequency=Pixel|Compute, Derivative=Zero, bIsInlined=true)
 +ExternalCodeDeclarations=(ReturnType=Float3,			Name="VertexTangent",					Definition="Parameters.TangentToWorld[0]", bIsInlined=true, EnvironmentDefines=((Name="USES_TRANSFORM_VECTOR", ShaderFrequency=Pixel|Compute)))
 
diff --git a/Engine/Source/Runtime/Engine/Private/Materials/HLSLMaterialTranslator.cpp b/Engine/Source/Runtime/Engine/Private/Materials/HLSLMaterialTranslator.cpp
--- a/Engine/Source/Runtime/Engine/Private/Materials/HLSLMaterialTranslator.cpp
+++ b/Engine/Source/Runtime/Engine/Private/Materials/HLSLMaterialTranslator.cpp
@@ -6414,13 +6414,13 @@ int32 FHLSLMaterialTranslator::ParticleSubUVProperty(int32 PropertyIndex)
 	switch (PropertyIndex)
 	{
 	case 0:
-		Result = AddCodeChunk(MCT_Float2, TEXT("Parameters.Particle.SubUVCoords[0].xy"));
+		Result = AddExternalCodeChunk(TEXT("ParticleSubUVCoords0"));
 		break;
 	case 1:
-		Result = AddCodeChunk(MCT_Float2, TEXT("Parameters.Particle.SubUVCoords[1].xy"));
+		Result = AddExternalCodeChunk(TEXT("ParticleSubUVCoords1"));
 		break;
 	case 2:
-		Result = AddCodeChunk(MCT_Float, TEXT("Parameters.Particle.SubUVLerp"));
+		Result = AddExternalCodeChunk(TEXT("ParticleSubUVLerp"));
 		break;
 	default:
 		checkNoEntry();
@@ -6719,17 +6719,12 @@ int32 FHLSLMaterialTranslator::ObjectWorldPosition(EPositionOrigin OriginType)
 
 int32 FHLSLMaterialTranslator::ObjectRadius()
 {
-	return GetPrimitiveProperty(MCT_Float, TEXT("ObjectRadius"), TEXT("ObjectRadius"));		
+	return AddExternalCodeChunk(TEXT("ObjectRadius"));
 }
 
 int32 FHLSLMaterialTranslator::ObjectBounds()
 {
-	if (!CheckPrimitivePropertyCompatibity(ANSI_TO_TCHAR(__FUNCTION__)))
-	{
-		return INDEX_NONE;
-	}
-
-	return AddInlinedCodeChunkZeroDeriv(MCT_Float3, TEXT("float3(GetPrimitiveData(Parameters).ObjectBoundsX, GetPrimitiveData(Parameters).ObjectBoundsY, GetPrimitiveData(Parameters).ObjectBoundsZ)"));
+	return AddExternalCodeChunk(TEXT("ObjectBounds"));
 }
 
 int32 FHLSLMaterialTranslator::ObjectLocalBounds(int32 OutputIndex)
diff --git a/Engine/Source/Runtime/Engine/Private/Materials/MaterialExpressions.cpp b/Engine/Source/Runtime/Engine/Private/Materials/MaterialExpressions.cpp
--- a/Engine/Source/Runtime/Engine/Private/Materials/MaterialExpressions.cpp
+++ b/Engine/Source/Runtime/Engine/Private/Materials/MaterialExpressions.cpp
@@ -12189,16 +12189,6 @@ void UMaterialExpressionObjectPositionWS::GetExpressionToolTip(TArray<FString>&
 // UMaterialExpressionObjectRadius
 ///////////////////////////////////////////////////////////////////////////////
 #if WITH_EDITOR
-int32 UMaterialExpressionObjectRadius::Compile(class FMaterialCompiler* Compiler, int32 OutputIndex)
-{
-	if (Material && Material->MaterialDomain == MD_DeferredDecal)
-	{
-		return CompilerError(Compiler, TEXT("Expression not available in the deferred decal material domain."));
-	}
-
-	return Compiler->ObjectRadius();
-}
-
 void UMaterialExpressionObjectRadius::GetCaption(TArray<FString>& OutCaptions) const
 {
 	OutCaptions.Add(TEXT("Object Radius"));
@@ -12209,16 +12199,6 @@ void UMaterialExpressionObjectRadius::GetCaption(TArray<FString>& OutCaptions) c
 // UMaterialExpressionObjectBoundingBox
 ///////////////////////////////////////////////////////////////////////////////
 #if WITH_EDITOR
-int32 UMaterialExpressionObjectBounds::Compile(class FMaterialCompiler* Compiler, int32 OutputIndex)
-{
-	if (Material && Material->MaterialDomain == MD_DeferredDecal)
-	{
-		return CompilerError(Compiler, TEXT("Expression not available in the deferred decal material domain."));
-	}
-
-	return Compiler->ObjectBounds();
-}
-
 void UMaterialExpressionObjectBounds::GetCaption(TArray<FString>& OutCaptions) const
 {
 	OutCaptions.Add(TEXT("Object Bounds"));
diff --git a/Engine/Source/Runtime/Engine/Private/Materials/MaterialExpressionsIR.cpp b/Engine/Source/Runtime/Engine/Private/Materials/MaterialExpressionsIR.cpp
--- a/Engine/Source/Runtime/Engine/Private/Materials/MaterialExpressionsIR.cpp
+++ b/Engine/Source/Runtime/Engine/Private/Materials/MaterialExpressionsIR.cpp
@@ -630,28 +630,37 @@ void UMaterialExpressionTangent::Build(MIR::FEmitter& Emitter)
 	BuildUnaryArithmeticOperatorFromInput(Emitter, MIR::UO_Tan, &Input);
 }
 
-void UMaterialExpressionExternalCodeBase::Build(MIR::FEmitter& Emitter)
+static MIR::FValue* EmitInlineHLSL(const UMaterialExpressionExternalCodeBase& InExternalCodeExpression, MIR::FEmitter& Emitter, int32 InExternalCodeIdentifierIndex, TConstArrayView<MIR::FValue*> InArguments)
 {
-	if (ExternalCodeIdentifiers.IsEmpty())
+	checkf(
+		InExternalCodeIdentifierIndex >= 0 && InExternalCodeIdentifierIndex < InExternalCodeExpression.ExternalCodeIdentifiers.Num(),
+		TEXT("External code identifier index (%d) out of bounds; Upper bound is %d"), InExternalCodeIdentifierIndex, InExternalCodeExpression.ExternalCodeIdentifiers.Num()
+	);
+
+	const FName& ExternalCodeIdentifier = InExternalCodeExpression.ExternalCodeIdentifiers[InExternalCodeIdentifierIndex];
+	const FMaterialExternalCodeDeclaration* ExternalCodeDeclaration = MaterialExternalCodeRegistry::Get().FindExternalCode(ExternalCodeIdentifier);
+	if (!ExternalCodeDeclaration)
 	{
-		Emitter.Error(TEXT("Material expression is missing external code identifiers"));
-		return;
+		Emitter.Errorf(TEXT("Missing external code declaration for '%s' [Index=%d]"), *ExternalCodeIdentifier.ToString(), InExternalCodeIdentifierIndex);
+		return nullptr;
 	}
 
-	for (int32 OutputIndex = 0; OutputIndex < ExternalCodeIdentifiers.Num(); ++OutputIndex)
+	return Emitter.EmitInlineHLSL(ExternalCodeDeclaration, InArguments);
+}
+
+static void BuildInlineHLSLOutput(const UMaterialExpressionExternalCodeBase& InExternalCodeExpression, MIR::FEmitter& Emitter, TConstArrayView<MIR::FValue*> InArguments)
+{
+	for (int32 OutputIndex = 0; OutputIndex < InExternalCodeExpression.ExternalCodeIdentifiers.Num(); ++OutputIndex)
 	{
-		if (const FMaterialExternalCodeDeclaration* ExternalCodeDeclaration = MaterialExternalCodeRegistry::Get().FindExternalCode(ExternalCodeIdentifiers[OutputIndex]))
-		{
-			MIR::FValue* Value = Emitter.EmitInlineHLSL(ExternalCodeDeclaration);
-			Emitter.Output(OutputIndex, Value);
-		}
-		else
-		{
-			Emitter.Errorf(TEXT("Missing external code declaration for '%s' [OutputIndex=%d]"), *ExternalCodeIdentifiers[OutputIndex].ToString(), OutputIndex);
-		}
+		Emitter.Output(OutputIndex, EmitInlineHLSL(InExternalCodeExpression, Emitter, OutputIndex, InArguments));
 	}
 }
 
+void UMaterialExpressionExternalCodeBase::Build(MIR::FEmitter& Emitter)
+{
+	BuildInlineHLSLOutput(*this, Emitter, {});
+}
+
 /* Binary Operators */
 
 static void BuildBinaryOperatorFromValue(
@@ -782,6 +791,24 @@ void UMaterialExpressionCrossProduct::Build(MIR::FEmitter& Emitter)
 	BuildBinaryOperatorFromInput(Emitter, MIR::BO_Cross, &A, &B);
 }
 
+void UMaterialExpressionEyeAdaptationInverse::Build(MIR::FEmitter& Emitter)
+{
+	MIR::FValue* LightValue = Emitter
+		.DefaultInputToFloat3(&LightValueInput, { 1.0f, 1.0f, 1.0f })
+		.InputFloatVector(&LightValueInput, 3);
+
+	MIR::FValue* AlphaValue = Emitter
+		.DefaultInputToFloatScalar(&AlphaInput, 1.0f)
+		.InputFloatScalar(&AlphaInput);
+
+	UE_MIR_CHECKPOINT();
+
+	check(ExternalCodeIdentifiers.Num() == 1);
+	MIR::FValue* MultiplierValue = EmitInlineHLSL(*this, Emitter, 0, { AlphaValue });
+
+	Emitter.Output(0, Emitter.EmitOperator(MIR::BO_Multiply, LightValue, MultiplierValue));
+}
+
 void UMaterialExpressionOneMinus::Build(MIR::FEmitter& Emitter)
 {
 	// Default input to zero if not connected, then get it as a primitive.
@@ -871,6 +898,17 @@ void UMaterialExpressionIf::Build(MIR::FEmitter& Emitter)
 	Emitter.Output(0, OutputValue);
 }
 
+void UMaterialExpressionSphericalParticleOpacity::Build(MIR::FEmitter& Emitter)
+{
+	MIR::FValue* DensityValue = Emitter
+		.DefaultInputToFloatScalar(&Density, ConstantDensity)
+		.Input(&Density);
+		
+	UE_MIR_CHECKPOINT();
+
+	BuildInlineHLSLOutput(*this, Emitter, { DensityValue });
+}
+
 void UMaterialExpressionTextureObject::Build(MIR::FEmitter& Emitter)
 {
 	Emitter.Output(0, Emitter.EmitTextureObject(Texture, SamplerType));
@@ -1021,6 +1059,30 @@ void UMaterialExpressionTextureCoordinate::Build(MIR::FEmitter& Emitter)
 	Emitter.Output(GetOutput(0), OutputValue);
 }
 
+void UMaterialExpressionTime::Build(MIR::FEmitter& Emitter)
+{
+	const MIR::FType* ScalarFloatType = MIR::FPrimitiveType::GetScalar(MIR::ScalarKind_Float);
+
+	// When pausing the game is ignored for this time expression, use real-time instead of game-time.
+	const TCHAR* TimeFieldName = bIgnorePause ? TEXT("View.<PREV>RealTime") : TEXT("View.<PREV>GameTime");
+
+	if (!bOverride_Period)
+	{
+		Emitter.Output(0, Emitter.EmitInlineHLSL(ScalarFloatType, TimeFieldName, MIR::EValueFlags::SubstituteTagsInInlineHLSL));
+	}
+	else if (Period == 0.0f)
+	{
+		Emitter.Output(0, Emitter.EmitConstantFloatScalar(0.0f));
+	}
+	else
+	{
+		// Note: Don't use IR intrinsic for Fmod() here to avoid conversion to fp16 on mobile.
+		// We want full 32 bit float precision until the fmod when using a period.
+		MIR::FValue* PeriodValue = Emitter.EmitConstantFloatScalar(Period);
+		Emitter.Output(0, Emitter.EmitInlineHLSL(ScalarFloatType, FString::Printf(TEXT("fmod(%s,$0)"), TimeFieldName), { PeriodValue }, MIR::EValueFlags::SubstituteTagsInInlineHLSL));
+	}
+}
+
 void UMaterialExpressionReroute::Build(MIR::FEmitter& Emitter)
 {
 	MIR::FValue* Value = Emitter.Input(&Input);
diff --git a/Engine/Source/Runtime/Engine/Private/Materials/MaterialIR.cpp b/Engine/Source/Runtime/Engine/Private/Materials/MaterialIR.cpp
--- a/Engine/Source/Runtime/Engine/Private/Materials/MaterialIR.cpp
+++ b/Engine/Source/Runtime/Engine/Private/Materials/MaterialIR.cpp
@@ -49,7 +49,7 @@ const TCHAR* LexToString(EValueKind Kind)
 		case VK_Subscript: return TEXT("Subscript");
 		case VK_Cast: return TEXT("Cast");
 		case VK_TextureSample: return TEXT("TextureSample");
-		case VK_InlineHLSL: return TEXT("ExternalCode");
+		case VK_InlineHLSL: return TEXT("InlineHLSL");
 		case VK_StageSwitch: return TEXT("StageSwitch");
 		case VK_HardwarePartialDerivative: return TEXT("HardwarePartialDerivative");
 		/* invalid entries */
@@ -164,48 +164,49 @@ TConstArrayView<FValue*> FValue::GetUses() const
 		case VK_SetMaterialOutput:
 		{
 			auto This = static_cast<const FSetMaterialOutput*>(this);
-			return { &This->Arg, FSetMaterialOutput::NumUses };
+			return { &This->Arg, FSetMaterialOutput::NumStaticUses };
 		}
 
 		case VK_Operator:
 		{
 			auto This = static_cast<const FOperator*>(this);
-			return { &This->AArg, FOperator::NumUses };
+			return { &This->AArg, FOperator::NumStaticUses };
 		}
 
 		case VK_Branch:
 		{
 			auto This = static_cast<const FBranch*>(this);
-			return { &This->ConditionArg, FBranch::NumUses };
+			return { &This->ConditionArg, FBranch::NumStaticUses };
 		}
 
 		case VK_Subscript:
 		{
 			auto This = static_cast<const FSubscript*>(this);
-			return { &This->Arg, FSubscript::NumUses };
+			return { &This->Arg, FSubscript::NumStaticUses };
 		}
 
 		case VK_TextureSample:
 		{
 			auto This = static_cast<const FTextureSample*>(this);
-			return { &This->TextureObject, FTextureSample::NumUses };
+			return { &This->TextureObject, FTextureSample::NumStaticUses };
 		}
 
 		case VK_InlineHLSL:
 		{
-			return {}; // @todo-laura.hermanns - Add input arguments as soon as external code expressions support parameters
+			auto This = static_cast<const FInlineHLSL*>(this);
+			return { This->Arguments, This->NumArguments };
 		}
 
 		case VK_StageSwitch:
 		{
 			auto This = static_cast<const FStageSwitch*>(this);
-			return { This->Args, FStageSwitch::NumUses * NumStages };
+			return { This->Args, FStageSwitch::NumStaticUses * NumStages };
 		}
 
 		case VK_HardwarePartialDerivative:
 		{
 			auto This = static_cast<const FHardwarePartialDerivative*>(this);
-			return { &This->Arg, FHardwarePartialDerivative::NumUses };
+			return { &This->Arg, FHardwarePartialDerivative::NumStaticUses };
 		}
 
 		default: UE_MIR_UNREACHABLE();
@@ -216,7 +217,7 @@ TConstArrayView<FValue*> FValue::GetUsesForStage(MIR::EStage Stage) const
 {
 	if (const FStageSwitch* This = As<FStageSwitch>())
 	{
-		return { &This->Args[int32(Stage)], FStageSwitch::NumUses };
+		return { &This->Args[int32(Stage)], FStageSwitch::NumStaticUses };
 	}
 	return GetUses();
 }
diff --git a/Engine/Source/Runtime/Engine/Private/Materials/MaterialIREmitter.cpp b/Engine/Source/Runtime/Engine/Private/Materials/MaterialIREmitter.cpp
--- a/Engine/Source/Runtime/Engine/Private/Materials/MaterialIREmitter.cpp
+++ b/Engine/Source/Runtime/Engine/Private/Materials/MaterialIREmitter.cpp
@@ -210,6 +210,16 @@ FEmitter& FEmitter::DefaultInputToFloatScalar(const FExpressionInput* InInput, T
 	return *this;
 }
 
+FEmitter& FEmitter::DefaultInputToFloat3(const FExpressionInput* InInput, FVector3f Value)
+{
+	if (!InInput->IsConnected())
+	{
+		Internal::BindValueToExpressionInput(BuilderImpl, InInput, EmitConstantFloat3(Value));
+	}
+
+	return *this;
+}
+
 FEmitter& FEmitter::DefaultInputToFloat4(const FExpressionInput* InInput, FVector4f Value)
 {
 	if (!InInput->IsConnected())
@@ -305,6 +315,13 @@ FValue* FEmitter::InputPrimitive(const FExpressionInput* InInput)
 	return Value;
 }
 
+FValue* FEmitter::InputFloatVector(const FExpressionInput* InInput, int32 NumComponents)
+{
+	FValue* Value = Input(InInput);
+	CheckInputIsVector(InInput, Value, NumComponents, ScalarKind_Float);
+	return Value;
+}
+
 FValue* FEmitter::InputTexture(const FExpressionInput* InInput)
 {
 	FValue* Value = Input(InInput);
@@ -375,6 +392,20 @@ void FEmitter::CheckInputIsScalar(const FExpressionInput* InInput, FValue* Input
 	}
 }
 
+void FEmitter::CheckInputIsVector(const FExpressionInput* InInput, FValue* InputValue, int32 NumComponents, EScalarKind Kind)
+{
+	if (IsAnyNonNullAPoison(InputValue))
+	{
+		return;
+	}
+
+	const FPrimitiveType* PrimitiveType = InputValue->Type->AsPrimitive();
+	if (!PrimitiveType || !PrimitiveType->IsVector() || PrimitiveType->NumRows != NumComponents || PrimitiveType->ScalarKind != Kind)
+	{
+		Errorf(TEXT("Input '%s' expected to be a %s %d-component vector. It is %s instead."), *SlowFindInputName(Expression, InInput).ToString(), ScalarKindToString(Kind), NumComponents, InputValue->Type->GetSpelling().GetData());
+	}
+}
+
 void FEmitter::CheckInputTypeIs(const FExpressionInput* InInput, FValue* InputValue, ETypeKind Kind)
 {
 	if (IsAnyNonNullAPoison(InputValue))
@@ -1987,17 +2018,29 @@ FValue* FEmitter::EmitTextureSample(FValue* Texture, FValue* TexCoord, FValue* T
 	return EmitPrototype(this, Prototype);
 }
 
-FValue* FEmitter::EmitInlineHLSL(const FType* Type, FString Code, EValueFlags ValueFlags, EGraphProperties UsedGraphProperties)
+FValue* FEmitter::EmitInlineHLSL(const FType* Type, FString Code, TConstArrayView<FValue*> InArguments, EValueFlags ValueFlags, EGraphProperties UsedGraphProperties)
 {
 	MIR::FInlineHLSL Prototype = MakePrototype<MIR::FInlineHLSL>(Type);
 	Prototype.Flags = ValueFlags | EValueFlags::HasDynamicHLSLCode;
 	Prototype.GraphProperties = UsedGraphProperties;
 	Prototype.Code = Module->PushUserString(MoveTemp(Code));
 
+	if (!InArguments.IsEmpty())
+	{
+		checkf(InArguments.Num() < MIR::FInlineHLSL::MaxNumArguments, TEXT("Number of arguments for inline-HLSL out of bounds: %d was specified, but upper bound is %d"), InArguments.Num(), MIR::FInlineHLSL::MaxNumArguments);
+		Prototype.NumArguments = InArguments.Num();
+		FMemory::Memcpy(&Prototype.Arguments[0], InArguments.GetData(), sizeof(FValue*) * InArguments.Num());
+	}
+
 	return EmitPrototype(this, Prototype);
 }
 
-FValue* FEmitter::EmitInlineHLSL(const FMaterialExternalCodeDeclaration* InExternalCodeDeclaration, EGraphProperties UsedGraphProperties)
+FValue* FEmitter::EmitInlineHLSL(const FType* Type, FString Code, EValueFlags ValueFlags, EGraphProperties UsedGraphProperties)
+{
+	return EmitInlineHLSL(Type, MoveTemp(Code), {}, ValueFlags, UsedGraphProperties);
+}
+
+FValue* FEmitter::EmitInlineHLSL(const FMaterialExternalCodeDeclaration* InExternalCodeDeclaration, TConstArrayView<FValue*> InArguments, EGraphProperties UsedGraphProperties)
 {
 	check(InExternalCodeDeclaration != nullptr);
 	const FType* ReturnType = FType::FromMaterialValueType(InExternalCodeDeclaration->GetReturnTypeValue());
@@ -2006,9 +2049,21 @@ FValue* FEmitter::EmitInlineHLSL(const FMaterialExternalCodeDeclaration* InExter
 	Prototype.GraphProperties = UsedGraphProperties;
 	Prototype.ExternalCodeDeclaration = InExternalCodeDeclaration;
 
+	if (!InArguments.IsEmpty())
+	{
+		checkf(InArguments.Num() < MIR::FInlineHLSL::MaxNumArguments, TEXT("Number of arguments for inline-HLSL out of bounds: %d was specified, but upper bound is %d"), InArguments.Num(), MIR::FInlineHLSL::MaxNumArguments);
+		Prototype.NumArguments = InArguments.Num();
+		FMemory::Memcpy(&Prototype.Arguments[0], InArguments.GetData(), sizeof(FValue*) * InArguments.Num());
+	}
+
 	return EmitPrototype(this, Prototype);
 }
 
+FValue* FEmitter::EmitInlineHLSL(const FMaterialExternalCodeDeclaration* InExternalCodeDeclaration, EGraphProperties UsedGraphProperties)
+{
+	return EmitInlineHLSL(InExternalCodeDeclaration, {}, UsedGraphProperties);
+}
+
 const FType* FEmitter::TryGetCommonType(const FType* A, const FType* B)
 {
 	// Trivial case: types are equal
diff --git a/Engine/Source/Runtime/Engine/Private/Materials/MaterialIRModuleBuilder.cpp b/Engine/Source/Runtime/Engine/Private/Materials/MaterialIRModuleBuilder.cpp
--- a/Engine/Source/Runtime/Engine/Private/Materials/MaterialIRModuleBuilder.cpp
+++ b/Engine/Source/Runtime/Engine/Private/Materials/MaterialIRModuleBuilder.cpp
@@ -16,6 +16,7 @@
 #include "Materials/MaterialExpressionFunctionOutput.h"
 #include "Materials/MaterialExpressionMaterialFunctionCall.h"
 #include "Materials/Material.h"
+#include "MaterialDomain.h"
 #include "MaterialExpressionIO.h"
 #include "MaterialShared.h"
 #include "Materials/MaterialExpression.h"
@@ -422,6 +423,29 @@ struct FMaterialIRModuleBuilderImpl
 		Module->EnvironmentDefines = MoveTemp(ValueAnalyzer.EnvironmentDefines);
 	}
 
+	void Step_AnalyzeBuiltinDefines()
+	{
+		// Match various defines against the material configuration
+		if (Module->EnvironmentDefines.Contains(TEXT("MIR.SceneDepth")))
+		{
+			if (ValueAnalyzer.Material->MaterialDomain != MD_PostProcess && !IsTranslucentBlendMode(ValueAnalyzer.Material->BlendMode))
+			{
+				Module->AddError(nullptr, TEXT("Only transparent or postprocess materials can read from scene depth."));
+			}
+		}
+
+		// Remove all environment defines that have the "MIR." prefix as they are not meant to propagate into the set of compiler environment defines.
+		TCHAR DefineMIRPrefix[5] = {};
+		for (auto Iter = Module->EnvironmentDefines.CreateIterator(); Iter; ++Iter)
+		{
+			Iter->ToStringTruncate(DefineMIRPrefix, UE_ARRAY_COUNT(DefineMIRPrefix));
+			if (FCString::Strncmp(DefineMIRPrefix, TEXT("MIR."), UE_ARRAY_COUNT(DefineMIRPrefix)) == 0)
+			{
+				Iter.RemoveCurrent();
+			}
+		}
+	}
+
 	void Step_LinkInstructions()
 	{
 		TArray<MIR::FInstruction*> InstructionStack{};
@@ -560,6 +584,7 @@ bool FMaterialIRModuleBuilder::Build(FMaterialIRModule* TargetModule)
 	Impl.Step_FlowValuesIntoMaterialOutputs();
 	Impl.Step_AnalyzeIRGraph();
 	Impl.Step_ConsolidateEnvironmentDefines();
+	Impl.Step_AnalyzeBuiltinDefines();
 	Impl.Step_LinkInstructions();
 	Impl.Step_Finalize();
 
diff --git a/Engine/Source/Runtime/Engine/Private/Materials/MaterialIRToHLSLTranslator.cpp b/Engine/Source/Runtime/Engine/Private/Materials/MaterialIRToHLSLTranslator.cpp
--- a/Engine/Source/Runtime/Engine/Private/Materials/MaterialIRToHLSLTranslator.cpp
+++ b/Engine/Source/Runtime/Engine/Private/Materials/MaterialIRToHLSLTranslator.cpp
@@ -51,6 +51,12 @@ struct FHLSLPrinter
 		return *this;
 	}
 
+	FHLSLPrinter& operator<<(const FStringView& Text)
+	{
+		Buffer.Append(Text);
+		return *this;
+	}
+
 	FHLSLPrinter& operator<<(int32 Value)
 	{
 		Buffer.Appendf(TEXT("%d"), Value);
@@ -211,7 +217,6 @@ static bool IsOperatorInfix(MIR::EOperator Op)
 		case MIR::BO_Subtract:
 		case MIR::BO_Multiply:
 		case MIR::BO_Divide:
-		case MIR::BO_Fmod:
 		case MIR::BO_Max:
 		case MIR::BO_Min:
 			return true;
@@ -419,8 +424,97 @@ struct FTranslator : FMaterialIRToHLSLTranslation
 		}
 	}
 
+	void LowerInlineHLSLWithArgumentsInternal(const FStringView& Code, const TConstArrayView<MIR::FValue*>& InArguments)
+	{
+		// Substitute argument tokens with instruction arguments
+		int32 ArgumentTokenStart = 0, ArgumentTokenEnd = 0;
+
+		// Prints the pending substring of the code and moves the range forward.
+		auto FlushCodeSubstring = [this, &ArgumentTokenEnd, &Code](int32 EndIndex) -> void
+			{
+				if (EndIndex > ArgumentTokenEnd)
+				{
+					Printer << Code.Mid(ArgumentTokenEnd, EndIndex - ArgumentTokenEnd);
+					ArgumentTokenEnd = EndIndex;
+				}
+			};
+
+		auto SubstituteNextArgument = [this, &ArgumentTokenStart, &ArgumentTokenEnd, &Code, &InArguments]() -> void
+			{
+				// Scan digits for argument index
+				int32 ArgumentIndexValue = 0;
+				int32 NumDigits = 0;
+
+				while (ArgumentTokenStart + NumDigits < Code.Len())
+				{
+					TCHAR CodeCharacter = Code[ArgumentTokenStart + NumDigits];
+					if (!FChar::IsDigit(CodeCharacter))
+					{
+						break;
+					}
+					ArgumentIndexValue *= 10;
+					ArgumentIndexValue += (CodeCharacter - TEXT('0'));
+					++NumDigits;
+				}
+
+				checkf(NumDigits > 0, TEXT("Failed to scan integer in inline-HLSL after token '$':\n\"%.*s\""), Code.Len(), Code.GetData());
+
+				checkf(
+					ArgumentIndexValue < InArguments.Num(), TEXT("Failed to substitute token $%d in inline-HLSL with given number of arguments (%d):\n\"%.*s\""),
+					ArgumentIndexValue, InArguments.Num(), Code.Len(), Code.GetData()
+				);
+
+				LowerValue(InArguments[ArgumentIndexValue]);
+
+				ArgumentTokenEnd = ArgumentTokenStart + NumDigits;
+			};
+
+		auto MatchCharacter = [&Code](int32& Position, TCHAR InCharacter) -> bool
+			{
+				if (Position < Code.Len() && Code[Position] == InCharacter)
+				{
+					++Position;
+					return true;
+				}
+				return false;
+			};
+
+		// Find all argument token characters '$'. For example "MyFunction($1, $0.xxxw)" can be subsituted with "MyFunction(MySecondArgument, MyFirstArgument.xxxw)"
+		while ((ArgumentTokenStart = Code.Find(TEXT("$"), ArgumentTokenEnd)) != INDEX_NONE)
+		{
+			FlushCodeSubstring(ArgumentTokenStart);
+
+			++ArgumentTokenStart;
+			if (MatchCharacter(ArgumentTokenStart, TEXT('{')))
+			{
+				SubstituteNextArgument();
+				const bool bMatchedArgumentToken = MatchCharacter(ArgumentTokenEnd, TEXT('}'));
+				checkf(bMatchedArgumentToken, TEXT("Failed to match argument token in inline-HLSL with syntax '${N}':\n\"%.*s\""), Code.Len(), Code.GetData());
+			}
+			else
+			{
+				SubstituteNextArgument();
+			}
+		}
+
+		FlushCodeSubstring(Code.Len());
+	}
+
 	void LowerInlineHLSL(const MIR::FInlineHLSL* ExternalCode)
 	{
+		auto PrintHLSLCode = [this, ExternalCode](const FStringView& Code) -> void
+			{
+				if (ExternalCode->NumArguments > 0)
+				{
+					check(ExternalCode->Arguments);
+					LowerInlineHLSLWithArgumentsInternal(Code, TConstArrayView<MIR::FValue*>(ExternalCode->Arguments, ExternalCode->NumArguments));
+				}
+				else
+				{
+					Printer << Code;
+				}
+			};
+
 		if (ExternalCode->HasFlags(MIR::EValueFlags::HasDynamicHLSLCode))
 		{
 			check(ExternalCode->Code != nullptr);
@@ -434,17 +528,17 @@ struct FTranslator : FMaterialIRToHLSLTranslation
 				constexpr bool bCompilingPreviousFrame = false;
 				Code.ReplaceInline(TEXT("<PREV>"), bCompilingPreviousFrame ? TEXT("Prev") : TEXT(""));
 	
-				Printer.Buffer.Append(MoveTemp(Code));
+				PrintHLSLCode(Code);
 			}
 			else
 			{
-				Printer.Buffer.Append(ExternalCode->Code);
+				PrintHLSLCode(ExternalCode->Code);
 			}
 		}
 		else
 		{
 			check(ExternalCode->ExternalCodeDeclaration != nullptr);
-			Printer.Buffer.Append(ExternalCode->ExternalCodeDeclaration->Definition);
+			PrintHLSLCode(ExternalCode->ExternalCodeDeclaration->Definition);
 		}
 	}
 
diff --git a/Engine/Source/Runtime/Engine/Public/Materials/MaterialExpression.h b/Engine/Source/Runtime/Engine/Public/Materials/MaterialExpression.h
--- a/Engine/Source/Runtime/Engine/Public/Materials/MaterialExpression.h
+++ b/Engine/Source/Runtime/Engine/Public/Materials/MaterialExpression.h
@@ -33,6 +33,7 @@ enum class EMaterialParameterType : uint8;
 namespace MIR
 {
 	class FEmitter;
+	struct FValue;
 }
 
 //@warning: FExpressionInput is mirrored in MaterialShared.h and manually "subclassed" in Material.h (FMaterialInput)
diff --git a/Engine/Source/Runtime/Engine/Public/Materials/MaterialExpressionEyeAdaptationInverse.h b/Engine/Source/Runtime/Engine/Public/Materials/MaterialExpressionEyeAdaptationInverse.h
--- a/Engine/Source/Runtime/Engine/Public/Materials/MaterialExpressionEyeAdaptationInverse.h
+++ b/Engine/Source/Runtime/Engine/Public/Materials/MaterialExpressionEyeAdaptationInverse.h
@@ -4,7 +4,7 @@
 
 #include "CoreMinimal.h"
 #include "UObject/ObjectMacros.h"
-#include "Materials/MaterialExpression.h"
+#include "Materials/MaterialExpressionExternalCodeBase.h"
 #include "MaterialExpressionEyeAdaptationInverse.generated.h"
 
 /**
@@ -13,12 +13,13 @@
 
 
 UCLASS(collapsecategories, hidecategories=Object)
-class UMaterialExpressionEyeAdaptationInverse : public UMaterialExpression
+class UMaterialExpressionEyeAdaptationInverse : public UMaterialExpressionExternalCodeBase
 {
 	GENERATED_UCLASS_BODY()
 
 	//~ Begin UMaterialExpression Interface
 #if WITH_EDITOR
+	virtual void Build(MIR::FEmitter& Emitter) override;
 	virtual int32 Compile(class FMaterialCompiler* Compiler, int32 OutputIndex) override;
 	virtual void GetCaption(TArray<FString>& OutCaptions) const override;
 #endif
diff --git a/Engine/Source/Runtime/Engine/Public/Materials/MaterialExpressionObjectBounds.h b/Engine/Source/Runtime/Engine/Public/Materials/MaterialExpressionObjectBounds.h
--- a/Engine/Source/Runtime/Engine/Public/Materials/MaterialExpressionObjectBounds.h
+++ b/Engine/Source/Runtime/Engine/Public/Materials/MaterialExpressionObjectBounds.h
@@ -5,11 +5,11 @@
 
 #include "CoreMinimal.h"
 #include "UObject/ObjectMacros.h"
-#include "Materials/MaterialExpression.h"
+#include "Materials/MaterialExpressionExternalCodeBase.h"
 #include "MaterialExpressionObjectBounds.generated.h"
 
 UCLASS(collapsecategories, hidecategories=Object)
-class UMaterialExpressionObjectBounds : public UMaterialExpression
+class UMaterialExpressionObjectBounds : public UMaterialExpressionExternalCodeBase
 {
 	GENERATED_BODY()
 
@@ -17,7 +17,6 @@ public:
 
 	//~ Begin UMaterialExpression Interface
 #if WITH_EDITOR
-	virtual int32 Compile(class FMaterialCompiler* Compiler, int32 OutputIndex) override;
 	virtual void GetCaption(TArray<FString>& OutCaptions) const override;
 #endif
 	//~ End UMaterialExpression Interface
diff --git a/Engine/Source/Runtime/Engine/Public/Materials/MaterialExpressionObjectRadius.h b/Engine/Source/Runtime/Engine/Public/Materials/MaterialExpressionObjectRadius.h
--- a/Engine/Source/Runtime/Engine/Public/Materials/MaterialExpressionObjectRadius.h
+++ b/Engine/Source/Runtime/Engine/Public/Materials/MaterialExpressionObjectRadius.h
@@ -5,11 +5,11 @@
 
 #include "CoreMinimal.h"
 #include "UObject/ObjectMacros.h"
-#include "Materials/MaterialExpression.h"
+#include "Materials/MaterialExpressionExternalCodeBase.h"
 #include "MaterialExpressionObjectRadius.generated.h"
 
 UCLASS(collapsecategories, hidecategories=Object)
-class UMaterialExpressionObjectRadius : public UMaterialExpression
+class UMaterialExpressionObjectRadius : public UMaterialExpressionExternalCodeBase
 {
 	GENERATED_BODY()
 
@@ -17,7 +17,6 @@ public:
 
 	//~ Begin UMaterialExpression Interface
 #if WITH_EDITOR
-	virtual int32 Compile(class FMaterialCompiler* Compiler, int32 OutputIndex) override;
 	virtual void GetCaption(TArray<FString>& OutCaptions) const override;
 #endif
 	//~ End UMaterialExpression Interface
diff --git a/Engine/Source/Runtime/Engine/Public/Materials/MaterialExpressionParticleSubUVProperties.h b/Engine/Source/Runtime/Engine/Public/Materials/MaterialExpressionParticleSubUVProperties.h
--- a/Engine/Source/Runtime/Engine/Public/Materials/MaterialExpressionParticleSubUVProperties.h
+++ b/Engine/Source/Runtime/Engine/Public/Materials/MaterialExpressionParticleSubUVProperties.h
@@ -5,11 +5,11 @@
 
 #include "CoreMinimal.h"
 #include "UObject/ObjectMacros.h"
-#include "Materials/MaterialExpression.h"
+#include "Materials/MaterialExpressionExternalCodeBase.h"
 #include "MaterialExpressionParticleSubUVProperties.generated.h"
 
 UCLASS(MinimalAPI)
-class UMaterialExpressionParticleSubUVProperties : public UMaterialExpression
+class UMaterialExpressionParticleSubUVProperties : public UMaterialExpressionExternalCodeBase
 {
 	GENERATED_UCLASS_BODY()
 
diff --git a/Engine/Source/Runtime/Engine/Public/Materials/MaterialExpressionSphericalParticleOpacity.h b/Engine/Source/Runtime/Engine/Public/Materials/MaterialExpressionSphericalParticleOpacity.h
--- a/Engine/Source/Runtime/Engine/Public/Materials/MaterialExpressionSphericalParticleOpacity.h
+++ b/Engine/Source/Runtime/Engine/Public/Materials/MaterialExpressionSphericalParticleOpacity.h
@@ -6,11 +6,11 @@
 #include "CoreMinimal.h"
 #include "UObject/ObjectMacros.h"
 #include "MaterialExpressionIO.h"
-#include "Materials/MaterialExpression.h"
+#include "Materials/MaterialExpressionExternalCodeBase.h"
 #include "MaterialExpressionSphericalParticleOpacity.generated.h"
 
 UCLASS(collapsecategories, hidecategories=Object)
-class UMaterialExpressionSphericalParticleOpacity : public UMaterialExpression
+class UMaterialExpressionSphericalParticleOpacity : public UMaterialExpressionExternalCodeBase
 {
 	GENERATED_UCLASS_BODY()
 
@@ -25,6 +25,7 @@ class UMaterialExpressionSphericalParticleOpacity : public UMaterialExpression
 
 	//~ Begin UMaterialExpression Interface
 #if WITH_EDITOR
+	virtual void Build(MIR::FEmitter& Emitter) override;
 	virtual int32 Compile(class FMaterialCompiler* Compiler, int32 OutputIndex) override;
 	virtual void GetCaption(TArray<FString>& OutCaptions) const override
 	{
diff --git a/Engine/Source/Runtime/Engine/Public/Materials/MaterialExpressionTime.h b/Engine/Source/Runtime/Engine/Public/Materials/MaterialExpressionTime.h
--- a/Engine/Source/Runtime/Engine/Public/Materials/MaterialExpressionTime.h
+++ b/Engine/Source/Runtime/Engine/Public/Materials/MaterialExpressionTime.h
@@ -29,6 +29,7 @@ public:
 
 	//~ Begin UMaterialExpression Interface
 #if WITH_EDITOR
+	virtual void Build(MIR::FEmitter& Emitter) override;
 	virtual int32 Compile(class FMaterialCompiler* Compiler, int32 OutputIndex) override;
 	virtual void GetCaption(TArray<FString>& OutCaptions) const override;
 	virtual bool NeedsRealtimePreview() override { return true; }
diff --git a/Engine/Source/Runtime/Engine/Public/Materials/MaterialIR.h b/Engine/Source/Runtime/Engine/Public/Materials/MaterialIR.h
--- a/Engine/Source/Runtime/Engine/Public/Materials/MaterialIR.h
+++ b/Engine/Source/Runtime/Engine/Public/Materials/MaterialIR.h
@@ -543,14 +543,14 @@ struct FInstruction : FValue
 	FBlock* GetDesiredBlockForUse(EStage Stage, int32 UseIndex);
 };
 
-template <EValueKind TTypeKind, uint32 TNumUses>
+template <EValueKind TTypeKind, uint32 TNumStaticUses>
 struct TInstruction : FInstruction
 {
 	/** The kind of this instruction. */
 	static constexpr EValueKind TypeKind = TTypeKind;
 
-	/** The number of uses this instruction has. */
-	static constexpr uint32 NumUses = TNumUses;
+	/** The number of values this instruction uses statically. Some instructions have a dynamic number of uses, in which case NumStaticUses is 0. */
+	static constexpr uint32 NumStaticUses = TNumStaticUses;
 };
 
 /**
@@ -597,6 +597,14 @@ struct FSetMaterialOutput : TInstruction<VK_SetMaterialOutput, 1>
  */
 struct FInlineHLSL : TInstruction<VK_InlineHLSL, 0>
 {
+	static constexpr int32 MaxNumArguments = 16;
+
+	/** Array of argument values. */
+	FValue* Arguments[MaxNumArguments];
+
+	/** Number of arguments from another expression node. */
+	int32 NumArguments;
+
 	union
 	{
 		/** The declaration this instruction refers to. */
diff --git a/Engine/Source/Runtime/Engine/Public/Materials/MaterialIREmitter.h b/Engine/Source/Runtime/Engine/Public/Materials/MaterialIREmitter.h
--- a/Engine/Source/Runtime/Engine/Public/Materials/MaterialIREmitter.h
+++ b/Engine/Source/Runtime/Engine/Public/Materials/MaterialIREmitter.h
@@ -55,7 +55,10 @@ public:
 	
 	//
 	FEmitter& DefaultInputToFloatScalar(const FExpressionInput* Input, TFloat Value);
-	
+
+	//
+	FEmitter& DefaultInputToFloat3(const FExpressionInput* Input, FVector3f Value);
+
 	//
 	FEmitter& DefaultInputToFloat4(const FExpressionInput* Input, FVector4f Value);
 
@@ -69,7 +72,7 @@ public:
 
 	// Gets and returns the value flowing into input with specified `Input`.
 	// If no value is flowing in (e.g. bacause it is not connected) it reports
-	// an error and returns null.
+	// an error and returns a poison value.
 	FValue* Input(const FExpressionInput* Input);
 	
 	// Gets and returns the value flowing into input with specified `Input`.
@@ -96,13 +99,17 @@ public:
 	// Before returning the value, it checks that it has a primitive type.
 	FValue* InputPrimitive(const FExpressionInput* Input);
 
+	// Gets and returns the value flowing into input with specified `Input`.
+	// Before returning the value, it checks that it is a float vector.
+	FValue* InputFloatVector(const FExpressionInput* Input, int32 NumComponents);
+
 	//
 	FValue* InputTexture(const FExpressionInput* Input);
 
 	/* Output */
 	
 	// Flows specified `Value` out of output with specified `OutputIndex`.
-	FEmitter& Output(int OutputIndex, FValue* Value);
+	FEmitter& Output(int32 OutputIndex, FValue* Value);
 
 	// Flows specified `Value` out of specified expression `Output`.
 	FEmitter& Output(const FExpressionOutput* Output, FValue* Value);
@@ -121,6 +128,9 @@ public:
 	//
 	void CheckInputIsScalar(const FExpressionInput* Input, FValue* InputValue, EScalarKind Kind);
 
+	//
+	void CheckInputIsVector(const FExpressionInput* Input, FValue* InputValue, int32 NumComponents, EScalarKind Kind);
+
 	// Checks that the type of the value `InputValue` flowing into `Input` is of
 	// specified type `Kind. If it isn't it reports an error. You may check whether
 	// an error occurred with `IsInvalid()`.
@@ -156,7 +166,7 @@ public:
 
 	/* Other Values */
 
-	FValue* EmitSubscript(FValue* Value, int ComponentIndex);
+	FValue* EmitSubscript(FValue* Value, int32 ComponentIndex);
 	FValue* EmitSwizzle(FValue* Value, FSwizzleMask Mask);
 	FValue* EmitTextureObject(UTexture* Texture, EMaterialSamplerType SamplerType);
 	FValue* EmitParameter(FName Name, FMaterialParameterMetadata& Metadata, EMaterialSamplerType SamplerType = SAMPLERTYPE_Color);
@@ -195,7 +205,9 @@ public:
 	 */
 	FValue* EmitTextureSample(FValue* Texture, FValue* TexCoord, FValue* TexCoordDdx, FValue* TexCoordDdy, FValue* MipValue, ESamplerSourceMode SamplerSourceMode, ETextureMipValueMode MipValueMode, EMaterialSamplerType SamplerType);
 	
+	FValue* EmitInlineHLSL(const FType* Type, FString Code, TConstArrayView<FValue*> InputValues, EValueFlags ValueFlags = EValueFlags::None, EGraphProperties UsedGraphProperties = EGraphProperties::None);
 	FValue* EmitInlineHLSL(const FType* Type, FString Code, EValueFlags ValueFlags = EValueFlags::None, EGraphProperties UsedGraphProperties = EGraphProperties::None);
+	FValue* EmitInlineHLSL(const FMaterialExternalCodeDeclaration* InExternalCodeDeclaration, TConstArrayView<FValue*> InputValues, EGraphProperties UsedGraphProperties = EGraphProperties::None);
 	FValue* EmitInlineHLSL(const FMaterialExternalCodeDeclaration* InExternalCodeDeclaration, EGraphProperties UsedGraphProperties = EGraphProperties::None);
 
 	/**
