Iris
* Pawns are now spatially filtered by default

#rnx

diff --git a/Engine/Config/BaseEngine.ini b/Engine/Config/BaseEngine.ini
--- a/Engine/Config/BaseEngine.ini
+++ b/Engine/Config/BaseEngine.ini
@@ -1428,12 +1428,12 @@ DefaultSpatialFilterName=Spatial
 ; Set the class-based dynamic filters configs here. Using None means the objects of this class are always relevant.
 !FilterConfigs=ClearArray
 ;+FilterConfigs=(ClassName=/Script/Engine.LevelScriptActor, DynamicFilterName=NotRouted) ; NotRouted means that this type will not be replicated
-+FilterConfigs=(ClassName=/Script/Engine.Actor, DynamicFilterName=None))
++FilterConfigs=(ClassName=/Script/Engine.Actor, DynamicFilterName=None)) ; By default classes derived from AActor are always relevant (because not filtered dynamically)
 ; Info types aren't supposed to have physical representation and rely on a static priority
 +FilterConfigs=(ClassName=/Script/Engine.Info, DynamicFilterName=None)
 +FilterConfigs=(ClassName=/Script/Engine.PlayerState, DynamicFilterName=None)
 ; Pawns can be spatially filtered...
-;+FilterConfigs=(ClassName=/Script/Engine.Pawn, DynamicFilterName=Spatial))
++FilterConfigs=(ClassName=/Script/Engine.Pawn, DynamicFilterName=Spatial))
 
 ; PrioritizerConfigs
 ; PlayerState is by default bAlwaysRelevant but the default player state prioritizer handles this.
