Added Skylight suppression paired with DFAO as it should be. Increased value from 0.7 to 0.9

#changelist validated

diff --git a/Engine/Config/BaseScalability.ini b/Engine/Config/BaseScalability.ini
--- a/Engine/Config/BaseScalability.ini
+++ b/Engine/Config/BaseScalability.ini
@@ -261,19 +261,24 @@ r.Shadow.Virtual.SMRT.SamplesPerRayLocal=8
 ;-----------------------------------------------------------------------------------------------------------------
 
 [GlobalIlluminationQuality@0]
+; When DFAO is disabled, Skylight intensity is suppressed
 r.DistanceFieldAO=0
+r.SkylightIntensityMultiplier=0.8
 r.Lumen.DiffuseIndirect.Allow=0
 r.RayTracing.Scene.BuildMode=0
 
 [GlobalIlluminationQuality@1]
-; Enable DFAO when Lumen is off
+; When DFAO is enabled, Skylight intensity is normal
 r.DistanceFieldAO=1
+r.SkylightIntensityMultiplier=1.0
 r.AOQuality=1
 r.Lumen.DiffuseIndirect.Allow=0
 r.RayTracing.Scene.BuildMode=0
 
 [GlobalIlluminationQuality@2]
+; When DFAO is enabled, Skylight intensity is normal
 r.DistanceFieldAO=1
+r.SkylightIntensityMultiplier=1.0
 r.AOQuality=1
 r.Lumen.DiffuseIndirect.Allow=1
 r.LumenScene.DirectLighting.MaxLightsPerTile=4
@@ -301,7 +306,9 @@ r.Lumen.TranslucencyVolume.RadianceCache.NumProbesToTraceBudget=70
 r.RayTracing.Scene.BuildMode=0
 
 [GlobalIlluminationQuality@3]
+; When DFAO is enabled, Skylight intensity is normal
 r.DistanceFieldAO=1
+r.SkylightIntensityMultiplier=1.0
 r.AOQuality=2
 r.Lumen.DiffuseIndirect.Allow=1
 r.LumenScene.DirectLighting.MaxLightsPerTile=8
@@ -330,7 +337,9 @@ r.Lumen.TranslucencyVolume.RadianceCache.NumProbesToTraceBudget=70
 r.RayTracing.Scene.BuildMode=1
 
 [GlobalIlluminationQuality@Cine]
+; When DFAO is enabled, Skylight intensity is normal
 r.DistanceFieldAO=1
+r.SkylightIntensityMultiplier=1.0
 r.AOQuality=2
 r.Lumen.DiffuseIndirect.Allow=1
 r.LumenScene.DirectLighting.MaxLightsPerTile=8
