[MIR] Support custom derivatives for external code declarations in new material translator.

This implements the following expressions in the new material translator: ComponentMask, PixelDepth, ScreenPosition, VertexColor, and VertexNormalWS.

#rnx

diff --git a/Engine/Config/BaseMaterialExpressions.ini b/Engine/Config/BaseMaterialExpressions.ini
--- a/Engine/Config/BaseMaterialExpressions.ini
+++ b/Engine/Config/BaseMaterialExpressions.ini
@@ -623,6 +623,7 @@ bShaderInputData=true
 
 [/Script/Engine.MaterialExpressionPixelDepth]
 +MenuCategories=NSLOCTEXT("MaterialExpression", "Depth", "Depth")
++ExternalCodeIdentifiers="PixelDepth"
 bShaderInputData=true
 
 [/Script/Engine.MaterialExpressionPixelNormalWS]
@@ -738,6 +739,8 @@ bShaderInputData=true
 
 [/Script/Engine.MaterialExpressionScreenPosition]
 +MenuCategories=NSLOCTEXT("MaterialExpression", "Coordinates", "Coordinates")
++ExternalCodeIdentifiers="ViewportUV"
++ExternalCodeIdentifiers="PixelPosition"
 bShaderInputData=true
 
 [/Script/Engine.MaterialExpressionSetMaterialAttributes]
@@ -938,6 +941,7 @@ bShaderInputData=true
 
 [/Script/Engine.MaterialExpressionVertexColor]
 +MenuCategories=NSLOCTEXT("MaterialExpression", "Constants", "Constants")
++ExternalCodeIdentifiers="VertexColor"
 bShaderInputData=true
 
 [/Script/Engine.MaterialExpressionVertexInterpolator]
@@ -946,6 +950,7 @@ bShaderInputData=true
 [/Script/Engine.MaterialExpressionVertexNormalWS]
 +MenuCategories=NSLOCTEXT("MaterialExpression", "Vectors", "Vectors")
 +MenuCategories=NSLOCTEXT("MaterialExpression", "Coordinates", "Coordinates")
++ExternalCodeIdentifiers="VertexNormal"
 bShaderInputData=true
 
 [/Script/Engine.MaterialExpressionVertexTangentWS]
@@ -991,6 +996,8 @@ bShaderInputData=true
 +MenuCategories=NSLOCTEXT("MaterialExpression", "Texture", "Texture")
 
 [/Script/Engine.MaterialExternalCodeCollection]
++ExternalCodeDeclarations=(ReturnType=Float,			Name="FooBar",					Definition="Test", DefinitionDDX="Test", Derivative=Valid)
+
 +ExternalCodeDeclarations=(ReturnType=Float3,			Name="AtmosphericLightVector",			Definition="MaterialExpressionAtmosphericLightVector(Parameters)")
 +ExternalCodeDeclarations=(ReturnType=Float3,			Name="AtmosphericLightColor",			Definition="MaterialExpressionAtmosphericLightColor(Parameters)")
 +ExternalCodeDeclarations=(ReturnType=Float3,			Name="CameraVector",					Definition="Parameters.CameraVector", ShaderFrequency=Pixel|Compute, bIsInlined=true, EnvironmentDefines=((Name="USES_TRANSFORM_VECTOR", ShaderFrequency=Pixel|Compute)))
@@ -1020,12 +1027,17 @@ bShaderInputData=true
 +ExternalCodeDeclarations=(ReturnType=Float2,			Name="ParticleSubUVCoords0",			Definition="Parameters.Particle.SubUVCoords[0].xy", EnvironmentDefines=((Name="USE_PARTICLE_SUBUVS")))
 +ExternalCodeDeclarations=(ReturnType=Float2,			Name="ParticleSubUVCoords1",			Definition="Parameters.Particle.SubUVCoords[1].xy", EnvironmentDefines=((Name="USE_PARTICLE_SUBUVS")))
 +ExternalCodeDeclarations=(ReturnType=Float,			Name="ParticleSubUVLerp",				Definition="Parameters.Particle.SubUVLerp", EnvironmentDefines=((Name="USE_PARTICLE_SUBUVS")))
++ExternalCodeDeclarations=(ReturnType=Float1,			Name="PixelDepth",						Definition="GetPixelDepth(Parameters)", DefinitionDDX="Parameters.ScreenPosition_DDX.w", DefinitionDDY="Parameters.ScreenPosition_DDY.w", Derivative=Valid, bIsInlined=true)
++ExternalCodeDeclarations=(ReturnType=Float2,			Name="PixelPosition",					Definition="GetPixelPosition(Parameters)", DefinitionDDX="float2(1.0f, 0.0f)", DefinitionDDY="float2(0.0f, 1.0f)", Derivative=Valid)
 +ExternalCodeDeclarations=(ReturnType=Float3,			Name="PreSkinnedPosition",				Definition="Parameters.PreSkinnedPosition", ShaderFrequency=Vertex, bIsInlined=true)
 +ExternalCodeDeclarations=(ReturnType=Float3,			Name="PreSkinnedNormal",				Definition="Parameters.PreSkinnedNormal", ShaderFrequency=Vertex, bIsInlined=true)
 +ExternalCodeDeclarations=(ReturnType=Float3,			Name="ReflectionVector",				Definition="Parameters.ReflectionVector", ShaderFrequency=Pixel|Compute, bIsInlined=true, EnvironmentDefines=((Name="USES_TRANSFORM_VECTOR", ShaderFrequency=Pixel|Compute)))
 +ExternalCodeDeclarations=(ReturnType=Float,			Name="SphericalParticleOpacity",		Definition="GetSphericalParticleOpacity(Parameters,$0)", ShaderFrequency=Pixel|Compute, EnvironmentDefines=((Name="NEEDS_PARTICLE_POSITION"), (Name="SPHERICAL_PARTICLE_OPACITY"), (Name="NEEDS_WORLD_POSITION_EXCLUDING_SHADER_OFFSETS"), (Name="MIR.SceneDepth")))
 +ExternalCodeDeclarations=(ReturnType=Float,			Name="TwoSidedSign",					Definition="Parameters.TwoSidedSign", ShaderFrequency=Pixel|Compute, Derivative=Zero, bIsInlined=true)
++ExternalCodeDeclarations=(ReturnType=Float4,			Name="VertexColor",						Definition="Parameters.VertexColor", DefinitionDDX="Parameters.VertexColor_DDX", DefinitionDDY="Parameters.VertexColor_DDY", Derivative=Valid, bIsInlined=true, EnvironmentDefines=((Name="INTERPOLATE_VERTEX_COLOR", ShaderFrequency=Pixel|Compute)))
++ExternalCodeDeclarations=(ReturnType=Float3,			Name="VertexNormal",					Definition="Parameters.TangentToWorld[2]", DefinitionDDX="Parameters.WorldGeoNormal_DDX", DefinitionDDY="Parameters.WorldGeoNormal_DDY", Derivative=Valid, bIsInlined=true, EnvironmentDefines=((Name="INTERPOLATE_VERTEX_COLOR", ShaderFrequency=Pixel|Compute)))
 +ExternalCodeDeclarations=(ReturnType=Float3,			Name="VertexTangent",					Definition="Parameters.TangentToWorld[0]", bIsInlined=true, EnvironmentDefines=((Name="USES_TRANSFORM_VECTOR", ShaderFrequency=Pixel|Compute)))
++ExternalCodeDeclarations=(ReturnType=Float2,			Name="ViewportUV",						Definition="GetViewportUV(Parameters)", DefinitionDDX="float2(View.ViewSizeAndInvSize.z, 0.0f)", DefinitionDDY="float2(0.0f, View.ViewSizeAndInvSize.w)", Derivative=Valid)
 
 +ExternalCodeDeclarations=(ReturnType=Float,			Name="GetPerInstanceRandom",			Definition="GetPerInstanceRandom(Parameters)", ShaderFrequency=Vertex|Pixel, Derivative=Zero, bIsInlined=true, EnvironmentDefines=((Name="NEEDS_PER_INSTANCE_RANDOM_PS", ShaderFrequency=Pixel)))
 +ExternalCodeDeclarations=(ReturnType=Float,			Name="GetPerInstanceFadeAmount",		Definition="GetPerInstanceFadeAmount(Parameters)", ShaderFrequency=Vertex|Pixel, Derivative=Zero, bIsInlined=true, EnvironmentDefines=((Name="USES_PER_INSTANCE_FADE_AMOUNT")))
diff --git a/Engine/Source/Runtime/Engine/Private/Materials/MaterialExpressionsIR.cpp b/Engine/Source/Runtime/Engine/Private/Materials/MaterialExpressionsIR.cpp
--- a/Engine/Source/Runtime/Engine/Private/Materials/MaterialExpressionsIR.cpp
+++ b/Engine/Source/Runtime/Engine/Private/Materials/MaterialExpressionsIR.cpp
@@ -648,6 +648,31 @@ void UMaterialExpressionTangent::Build(MIR::FEmitter& Emitter)
 	BuildPeriodicTrigonometricOperatorFromInput(Emitter, MIR::UO_Tan, &Input, Period);
 }
 
+void UMaterialExpressionComponentMask::Build(MIR::FEmitter& Emitter)
+{
+	MIR::FValue* Value = Emitter.Input(&Input);
+
+	MIR::FSwizzleMask Mask;
+	if (R)
+	{
+		Mask.Components[Mask.NumComponents++] = MIR::EVectorComponent::X;
+	}
+	if (G)
+	{
+		Mask.Components[Mask.NumComponents++] = MIR::EVectorComponent::Y;
+	}
+	if (B)
+	{
+		Mask.Components[Mask.NumComponents++] = MIR::EVectorComponent::Z;
+	}
+	if (A)
+	{
+		Mask.Components[Mask.NumComponents++] = MIR::EVectorComponent::W;
+	}
+
+	Emitter.Output(0, Emitter.EmitSwizzle(Value, Mask));
+}
+
 static MIR::FValue* EmitInlineHLSL(const UMaterialExpressionExternalCodeBase& InExternalCodeExpression, MIR::FEmitter& Emitter, int32 InExternalCodeIdentifierIndex, TConstArrayView<MIR::FValue*> InArguments)
 {
 	checkf(
@@ -1121,7 +1146,7 @@ void UMaterialExpressionTime::Build(MIR::FEmitter& Emitter)
 
 	if (!bOverride_Period)
 	{
-		Emitter.Output(0, Emitter.EmitInlineHLSL(ScalarFloatType, TimeFieldName, MIR::EValueFlags::SubstituteTagsInInlineHLSL));
+		Emitter.Output(0, Emitter.EmitInlineHLSL(ScalarFloatType, TimeFieldName, {}, MIR::EValueFlags::SubstituteTagsInInlineHLSL));
 	}
 	else if (Period == 0.0f)
 	{
@@ -1310,12 +1335,12 @@ static MIR::FValue* BuildViewProperty(MIR::FEmitter& Emitter, EMaterialExposedVi
 		// Fall back to compute the property's inverse from PropertyCode
 		return Emitter.EmitOperator(MIR::BO_Divide,
 			Emitter.EmitConstantFloatScalar(1.0f),
-			Emitter.EmitInlineHLSL(HLSLCodeType, MoveTemp(HLSLCode), MIR::EValueFlags::SubstituteTagsInInlineHLSL));
+			Emitter.EmitInlineHLSL(HLSLCodeType, MoveTemp(HLSLCode), {}, MIR::EValueFlags::SubstituteTagsInInlineHLSL));
 	}
 	else
 	{
 		// @todo-laura.hermanns - old translator used CastToNonLWCIfDisabled(), but LWC not supported in MIR yet
-		return Emitter.EmitInlineHLSL(HLSLCodeType, MoveTemp(HLSLCode), MIR::EValueFlags::SubstituteTagsInInlineHLSL);
+		return Emitter.EmitInlineHLSL(HLSLCodeType, MoveTemp(HLSLCode), {}, MIR::EValueFlags::SubstituteTagsInInlineHLSL);
 	}
 }
 
@@ -1340,7 +1365,7 @@ void UMaterialExpressionCameraPositionWS::Build(MIR::FEmitter& Emitter)
 
 void UMaterialExpressionPixelNormalWS::Build(MIR::FEmitter& Emitter)
 {
-	MIR::FValue* Output = Emitter.EmitInlineHLSL(MIR::FPrimitiveType::GetFloat3(), "Parameters.WorldNormal", {}, MIR::EGraphProperties::ReadsPixelNormal);
+	MIR::FValue* Output = Emitter.EmitInlineHLSL(MIR::FPrimitiveType::GetFloat3(), "Parameters.WorldNormal", {}, MIR::EValueFlags::None, MIR::EGraphProperties::ReadsPixelNormal);
 	Emitter.Output(0, Output);
 }
 
diff --git a/Engine/Source/Runtime/Engine/Private/Materials/MaterialExternalCodeRegistry.cpp b/Engine/Source/Runtime/Engine/Private/Materials/MaterialExternalCodeRegistry.cpp
--- a/Engine/Source/Runtime/Engine/Private/Materials/MaterialExternalCodeRegistry.cpp
+++ b/Engine/Source/Runtime/Engine/Private/Materials/MaterialExternalCodeRegistry.cpp
@@ -46,6 +46,56 @@ void FMaterialExternalCodeDeclaration::UpdateHash(FSHA1& Hasher) const
 }
 
 
+/*
+ * UMaterialExternalCodeCollection
+ */
+
+void UMaterialExternalCodeCollection::PostInitProperties()
+{
+	Super::PostInitProperties();
+
+	TStringBuilder<2048> InvalidPropertiesList;
+
+	auto AppendInvalidProperty = [&InvalidPropertiesList](const TCHAR* PropertyDescription) -> void
+		{
+			if (InvalidPropertiesList.Len() > 0)
+			{
+				InvalidPropertiesList.Append(TEXT(", "));
+			}
+			InvalidPropertiesList.Append(PropertyDescription);
+		};
+
+	for (const FMaterialExternalCodeDeclaration& ExternalCode : ExternalCodeDeclarations)
+	{
+		InvalidPropertiesList.Reset();
+
+		// Validate basic properties
+		if (ExternalCode.Definition.IsEmpty())
+		{
+			AppendInvalidProperty(TEXT("Definition is empty"));
+		}
+
+		// Validate derivative configuration properties
+		if (ExternalCode.Derivative == EDerivativeStatus::Valid)
+		{
+			if (ExternalCode.DefinitionDDX.IsEmpty())
+			{
+				AppendInvalidProperty(TEXT("DefinitionDDX is empty"));
+			}
+			if (ExternalCode.DefinitionDDY.IsEmpty())
+			{
+				AppendInvalidProperty(TEXT("DefinitionDDY is empty"));
+			}
+		}
+
+		if (InvalidPropertiesList.Len() > 0)
+		{
+			UE_LOG(LogMaterial, Error, TEXT("External HLSL code declaration '%s' is invalid: %s"), *ExternalCode.Name.ToString(), *InvalidPropertiesList);
+		}
+	}
+}
+
+
 /*
  * MaterialExternalCodeRegistry
  */
diff --git a/Engine/Source/Runtime/Engine/Private/Materials/MaterialExternalCodeRegistry.h b/Engine/Source/Runtime/Engine/Private/Materials/MaterialExternalCodeRegistry.h
--- a/Engine/Source/Runtime/Engine/Private/Materials/MaterialExternalCodeRegistry.h
+++ b/Engine/Source/Runtime/Engine/Private/Materials/MaterialExternalCodeRegistry.h
@@ -107,9 +107,22 @@ public:
 	UPROPERTY(Config)
 	FName Name;
 
+	/** Primary external code definition provided as HLSL shader code expression. */
 	UPROPERTY(Config)
 	FString Definition;
 
+	/** Secondary external code definition for custom DDX derivatives. */
+	UPROPERTY(Config)
+	FString DefinitionDDX;
+
+	/** Secondary external code definition for custom DDY derivatives. */
+	UPROPERTY(Config)
+	FString DefinitionDDY;
+
+	/**
+	 * Specifies the kind of derivative this code declaration provides.
+	 * If this is EDerivativeStatus::Valid, DefinitionDDX and DefinitionDDY provides the code definitions for the DDX and DDY derivatives respectively.
+	 */
 	UPROPERTY(Config)
 	EDerivativeStatus Derivative = EDerivativeStatus::NotAware;
 
@@ -147,6 +160,9 @@ public:
 	UPROPERTY(Config)
 	TArray<FMaterialExternalCodeDeclaration> ExternalCodeDeclarations;
 
+public:
+	virtual void PostInitProperties() override;
+
 };
 
 /** Helper struct for hard coded external code expressions for view properties. */
diff --git a/Engine/Source/Runtime/Engine/Private/Materials/MaterialIREmitter.cpp b/Engine/Source/Runtime/Engine/Private/Materials/MaterialIREmitter.cpp
--- a/Engine/Source/Runtime/Engine/Private/Materials/MaterialIREmitter.cpp
+++ b/Engine/Source/Runtime/Engine/Private/Materials/MaterialIREmitter.cpp
@@ -2153,17 +2153,13 @@ FValue* FEmitter::EmitInlineHLSL(const FType* Type, FString Code, TConstArrayVie
 	return EmitPrototype(this, Prototype);
 }
 
-FValue* FEmitter::EmitInlineHLSL(const FType* Type, FString Code, EValueFlags ValueFlags, EGraphProperties UsedGraphProperties)
-{
-	return EmitInlineHLSL(Type, MoveTemp(Code), {}, ValueFlags, UsedGraphProperties);
-}
-
-FValue* FEmitter::EmitInlineHLSL(const FMaterialExternalCodeDeclaration* InExternalCodeDeclaration, TConstArrayView<FValue*> InArguments, EGraphProperties UsedGraphProperties)
+FValue* FEmitter::EmitInlineHLSL(const FMaterialExternalCodeDeclaration* InExternalCodeDeclaration, TConstArrayView<FValue*> InArguments, EValueFlags ValueFlags, EGraphProperties UsedGraphProperties)
 {
 	check(InExternalCodeDeclaration != nullptr);
 	const FType* ReturnType = FType::FromMaterialValueType(InExternalCodeDeclaration->GetReturnTypeValue());
 
 	MIR::FInlineHLSL Prototype = MakePrototype<MIR::FInlineHLSL>(ReturnType);
+	Prototype.Flags = ValueFlags;
 	Prototype.GraphProperties = UsedGraphProperties;
 	Prototype.ExternalCodeDeclaration = InExternalCodeDeclaration;
 
@@ -2177,11 +2173,6 @@ FValue* FEmitter::EmitInlineHLSL(const FMaterialExternalCodeDeclaration* InExter
 	return EmitPrototype(this, Prototype);
 }
 
-FValue* FEmitter::EmitInlineHLSL(const FMaterialExternalCodeDeclaration* InExternalCodeDeclaration, EGraphProperties UsedGraphProperties)
-{
-	return EmitInlineHLSL(InExternalCodeDeclaration, {}, UsedGraphProperties);
-}
-
 const FType* FEmitter::TryGetCommonType(const FType* A, const FType* B)
 {
 	// Trivial case: types are equal
@@ -2264,6 +2255,26 @@ static FValue* DifferentiateExternalInput(FEmitter& Emitter, FExternalInput* Ext
 	return Emitter.EmitConstruct(PrimType, Emitter.EmitConstantScalarZero(PrimType->ScalarKind));
 }
 
+static FValue* EmitDifferentialInlineHLSL(FEmitter& Emitter, const FPrimitiveType* PrimitiveType, FInlineHLSL* InlineHLSL, EDerivativeAxis Axis)
+{
+	if (InlineHLSL->ExternalCodeDeclaration)
+	{
+		if (InlineHLSL->ExternalCodeDeclaration->Derivative == EDerivativeStatus::Valid)
+		{
+			const TConstArrayView<FValue*> InlineHLSLArguments(InlineHLSL->Arguments, InlineHLSL->NumArguments);
+			return Emitter.EmitInlineHLSL(InlineHLSL->ExternalCodeDeclaration, InlineHLSLArguments, Axis == EDerivativeAxis::X ? EValueFlags::DerivativeDDX : EValueFlags::DerivativeDDY);
+		}
+		else if (InlineHLSL->ExternalCodeDeclaration->Derivative == EDerivativeStatus::NotValid)
+		{
+			Emitter.Errorf(TEXT("External code declaration `%s` does not provide a valid derivative."), *InlineHLSL->ExternalCodeDeclaration->Name.ToString());
+			return FPoison::Get();
+		}
+	}
+
+	// If inline HLSL was dynamically generated and not provided from an external code declaration, simply return constant zero as no derivative is available.
+	return Emitter.EmitConstruct(PrimitiveType, Emitter.EmitConstantScalarZero(PrimitiveType->ScalarKind));
+}
+
 static FValue* DifferentiateOperator(FEmitter& Emitter, FOperator* Op, EDerivativeAxis Axis)
 {
 	const MIR::FPrimitiveType* PrimType = static_cast<const MIR::FPrimitiveType*>(Op->Type);
@@ -2526,9 +2537,14 @@ FValue* FEmitter::EmitAnalyticalPartialDerivative(FValue* Value, EDerivativeAxis
 		case MIR::VK_Constant:
 		case MIR::VK_UniformParameter:
 		case MIR::VK_TextureRead:
-		case MIR::VK_InlineHLSL:
 			return EmitConstruct(ValuePrimitiveType, EmitConstantScalarZero(ValuePrimitiveType->ScalarKind));
 
+		case MIR::VK_InlineHLSL:
+		{
+			FInlineHLSL* InlineHLSL = Value->As<FInlineHLSL>();
+			return EmitDifferentialInlineHLSL(*this, ValuePrimitiveType, InlineHLSL, Axis);
+		}
+
 		default:
 			UE_MIR_UNREACHABLE();
 	}
diff --git a/Engine/Source/Runtime/Engine/Private/Materials/MaterialIRToHLSLTranslator.cpp b/Engine/Source/Runtime/Engine/Private/Materials/MaterialIRToHLSLTranslator.cpp
--- a/Engine/Source/Runtime/Engine/Private/Materials/MaterialIRToHLSLTranslator.cpp
+++ b/Engine/Source/Runtime/Engine/Private/Materials/MaterialIRToHLSLTranslator.cpp
@@ -539,7 +539,18 @@ struct FTranslator : FMaterialIRToHLSLTranslation
 		else
 		{
 			check(ExternalCode->ExternalCodeDeclaration != nullptr);
-			PrintHLSLCode(ExternalCode->ExternalCodeDeclaration->Definition);
+			if (ExternalCode->HasFlags(MIR::EValueFlags::DerivativeDDX))
+			{
+				PrintHLSLCode(ExternalCode->ExternalCodeDeclaration->DefinitionDDX);
+			}
+			else if (ExternalCode->HasFlags(MIR::EValueFlags::DerivativeDDY))
+			{
+				PrintHLSLCode(ExternalCode->ExternalCodeDeclaration->DefinitionDDY);
+			}
+			else
+			{
+				PrintHLSLCode(ExternalCode->ExternalCodeDeclaration->Definition);
+			}
 		}
 	}
 
diff --git a/Engine/Source/Runtime/Engine/Public/Materials/MaterialExpressionComponentMask.h b/Engine/Source/Runtime/Engine/Public/Materials/MaterialExpressionComponentMask.h
--- a/Engine/Source/Runtime/Engine/Public/Materials/MaterialExpressionComponentMask.h
+++ b/Engine/Source/Runtime/Engine/Public/Materials/MaterialExpressionComponentMask.h
@@ -34,6 +34,7 @@ public:
 
 	//~ Begin UMaterialExpression Interface
 #if WITH_EDITOR
+	virtual void Build(MIR::FEmitter& Emitter) override;
 	virtual int32 Compile(class FMaterialCompiler* Compiler, int32 OutputIndex) override;
 	virtual void GetCaption(TArray<FString>& OutCaptions) const override;
 	virtual FText GetKeywords() const override { return FText::FromString(TEXT("component mask")); }
diff --git a/Engine/Source/Runtime/Engine/Public/Materials/MaterialExpressionPixelDepth.h b/Engine/Source/Runtime/Engine/Public/Materials/MaterialExpressionPixelDepth.h
--- a/Engine/Source/Runtime/Engine/Public/Materials/MaterialExpressionPixelDepth.h
+++ b/Engine/Source/Runtime/Engine/Public/Materials/MaterialExpressionPixelDepth.h
@@ -8,11 +8,11 @@
 
 #include "CoreMinimal.h"
 #include "UObject/ObjectMacros.h"
-#include "Materials/MaterialExpression.h"
+#include "Materials/MaterialExpressionExternalCodeBase.h"
 #include "MaterialExpressionPixelDepth.generated.h"
 
 UCLASS(collapsecategories, hidecategories=Object, MinimalAPI)
-class UMaterialExpressionPixelDepth : public UMaterialExpression
+class UMaterialExpressionPixelDepth : public UMaterialExpressionExternalCodeBase
 {
 	GENERATED_UCLASS_BODY()
 
diff --git a/Engine/Source/Runtime/Engine/Public/Materials/MaterialExpressionScreenPosition.h b/Engine/Source/Runtime/Engine/Public/Materials/MaterialExpressionScreenPosition.h
--- a/Engine/Source/Runtime/Engine/Public/Materials/MaterialExpressionScreenPosition.h
+++ b/Engine/Source/Runtime/Engine/Public/Materials/MaterialExpressionScreenPosition.h
@@ -5,12 +5,12 @@
 
 #include "CoreMinimal.h"
 #include "UObject/ObjectMacros.h"
-#include "Materials/MaterialExpression.h"
+#include "Materials/MaterialExpressionExternalCodeBase.h"
 #include "MaterialExpressionScreenPosition.generated.h"
 
 
 UCLASS(collapsecategories, hidecategories=Object)
-class UMaterialExpressionScreenPosition : public UMaterialExpression
+class UMaterialExpressionScreenPosition : public UMaterialExpressionExternalCodeBase
 {
 	GENERATED_UCLASS_BODY()
 
diff --git a/Engine/Source/Runtime/Engine/Public/Materials/MaterialExpressionVertexColor.h b/Engine/Source/Runtime/Engine/Public/Materials/MaterialExpressionVertexColor.h
--- a/Engine/Source/Runtime/Engine/Public/Materials/MaterialExpressionVertexColor.h
+++ b/Engine/Source/Runtime/Engine/Public/Materials/MaterialExpressionVertexColor.h
@@ -5,11 +5,11 @@
 
 #include "CoreMinimal.h"
 #include "UObject/ObjectMacros.h"
-#include "Materials/MaterialExpression.h"
+#include "Materials/MaterialExpressionExternalCodeBase.h"
 #include "MaterialExpressionVertexColor.generated.h"
 
 UCLASS(collapsecategories, hidecategories=Object, MinimalAPI)
-class UMaterialExpressionVertexColor : public UMaterialExpression
+class UMaterialExpressionVertexColor : public UMaterialExpressionExternalCodeBase
 {
 	GENERATED_UCLASS_BODY()
 
diff --git a/Engine/Source/Runtime/Engine/Public/Materials/MaterialExpressionVertexNormalWS.h b/Engine/Source/Runtime/Engine/Public/Materials/MaterialExpressionVertexNormalWS.h
--- a/Engine/Source/Runtime/Engine/Public/Materials/MaterialExpressionVertexNormalWS.h
+++ b/Engine/Source/Runtime/Engine/Public/Materials/MaterialExpressionVertexNormalWS.h
@@ -5,11 +5,11 @@
 
 #include "CoreMinimal.h"
 #include "UObject/ObjectMacros.h"
-#include "Materials/MaterialExpression.h"
+#include "Materials/MaterialExpressionExternalCodeBase.h"
 #include "MaterialExpressionVertexNormalWS.generated.h"
 
 UCLASS(MinimalAPI, collapsecategories, hidecategories=Object)
-class UMaterialExpressionVertexNormalWS : public UMaterialExpression
+class UMaterialExpressionVertexNormalWS : public UMaterialExpressionExternalCodeBase
 {
 	GENERATED_BODY()
 
diff --git a/Engine/Source/Runtime/Engine/Public/Materials/MaterialIR.h b/Engine/Source/Runtime/Engine/Public/Materials/MaterialIR.h
--- a/Engine/Source/Runtime/Engine/Public/Materials/MaterialIR.h
+++ b/Engine/Source/Runtime/Engine/Public/Materials/MaterialIR.h
@@ -57,6 +57,12 @@ enum class EValueFlags : uint8
 
 	/** Whether 'FInlineHLSL::Code' field is used. Otherwise, 'ExternalCodeDeclaration' is used. */
 	HasDynamicHLSLCode = 1 << 4, 
+
+	/** External code declaration must emit the secondary definition for the DDX derivative. */
+	DerivativeDDX = 1 << 5,
+
+	/** External code declaration must emit the secondary definition for the DDY derivative. */
+	DerivativeDDY = 1 << 6,
 };
 
 ENUM_CLASS_FLAGS(EValueFlags);
diff --git a/Engine/Source/Runtime/Engine/Public/Materials/MaterialIREmitter.h b/Engine/Source/Runtime/Engine/Public/Materials/MaterialIREmitter.h
--- a/Engine/Source/Runtime/Engine/Public/Materials/MaterialIREmitter.h
+++ b/Engine/Source/Runtime/Engine/Public/Materials/MaterialIREmitter.h
@@ -249,9 +249,7 @@ public:
 	FValue* EmitTextureSampleGrad(FValue* Texture, FValue* TexCoord, FValue* TexCoordDdx, FValue* TexCoordDdy, bool bAddViewMaterialMipBias, ESamplerSourceMode SamplerSourceMode, EMaterialSamplerType SamplerType);
 	
 	FValue* EmitInlineHLSL(const FType* Type, FString Code, TConstArrayView<FValue*> InputValues, EValueFlags ValueFlags = EValueFlags::None, EGraphProperties UsedGraphProperties = EGraphProperties::None);
-	FValue* EmitInlineHLSL(const FType* Type, FString Code, EValueFlags ValueFlags = EValueFlags::None, EGraphProperties UsedGraphProperties = EGraphProperties::None);
-	FValue* EmitInlineHLSL(const FMaterialExternalCodeDeclaration* InExternalCodeDeclaration, TConstArrayView<FValue*> InputValues, EGraphProperties UsedGraphProperties = EGraphProperties::None);
-	FValue* EmitInlineHLSL(const FMaterialExternalCodeDeclaration* InExternalCodeDeclaration, EGraphProperties UsedGraphProperties = EGraphProperties::None);
+	FValue* EmitInlineHLSL(const FMaterialExternalCodeDeclaration* InExternalCodeDeclaration, TConstArrayView<FValue*> InputValues, EValueFlags ValueFlags = EValueFlags::None, EGraphProperties UsedGraphProperties = EGraphProperties::None);
 
 	/**
 	 * @brief Emits the partial derivative of given value.
