[MIR] Remove debug code from BaseMaterialExpressions.ini file.

#rnx

diff --git a/Engine/Config/BaseMaterialExpressions.ini b/Engine/Config/BaseMaterialExpressions.ini
--- a/Engine/Config/BaseMaterialExpressions.ini
+++ b/Engine/Config/BaseMaterialExpressions.ini
@@ -996,8 +996,6 @@ bShaderInputData=true
 +MenuCategories=NSLOCTEXT("MaterialExpression", "Texture", "Texture")
 
 [/Script/Engine.MaterialExternalCodeCollection]
-+ExternalCodeDeclarations=(ReturnType=Float,			Name="FooBar",					Definition="Test", DefinitionDDX="Test", Derivative=Valid)
-
 +ExternalCodeDeclarations=(ReturnType=Float3,			Name="AtmosphericLightVector",			Definition="MaterialExpressionAtmosphericLightVector(Parameters)")
 +ExternalCodeDeclarations=(ReturnType=Float3,			Name="AtmosphericLightColor",			Definition="MaterialExpressionAtmosphericLightColor(Parameters)")
 +ExternalCodeDeclarations=(ReturnType=Float3,			Name="CameraVector",					Definition="Parameters.CameraVector", ShaderFrequency=Pixel|Compute, bIsInlined=true, EnvironmentDefines=((Name="USES_TRANSFORM_VECTOR", ShaderFrequency=Pixel|Compute)))
