VulkanRHI: bump minimum driver version requirements for VulkanRHI to include ray tracing fixes and new extensions.

#rnx

diff --git a/Engine/Config/BaseHardware.ini b/Engine/Config/BaseHardware.ini
--- a/Engine/Config/BaseHardware.ini
+++ b/Engine/Config/BaseHardware.ini
@@ -39,17 +39,19 @@ SuggestedDriverVersion="560.81"
 ; The newest driver which supports legacy Nvidia GPUs (700 series and older) is 474.64. We allow this driver for DX11.
 +DriverDenyList=(DriverVersion="<474.64", Reason="These driver versions have known stability issues and missing features", RHI="D3D11")
 
-; DX12 and Vulkan require recent drivers. The 536 series was released in June 2023.
+; DX12 require recent drivers. The 536 series was released in June 2023.
 +DriverDenyList=(DriverVersion="<536.00", Reason="These driver versions have known stability issues and missing features", RHI="D3D12")
-+DriverDenyList=(DriverVersion="<560.00", Reason="These driver versions have known stability issues and missing features", RHI="Vulkan")
+
+; Vulkan require recent drivers, the 570 series includes important fixes.
++DriverDenyList=(DriverVersion="<570.00", Reason="These driver versions have known stability issues and missing features", RHI="Vulkan")
 
 ;-----------------------------------------------------------------------------------------------------------------
 
 [GPU_NVIDIA Linux]
-SuggestedDriverVersion="535.86"
+SuggestedDriverVersion="570.86"
 
-; The 535 series of the Linux driver was released in June 2023.
-+DriverDenyList=(DriverVersion="<535.00", Reason="These driver versions have known stability issues and missing features")
+; Vulkan require recent drivers, the 570 series includes important fixes.
++DriverDenyList=(DriverVersion="<570.00", Reason="These driver versions have known stability issues and missing features")
 
 ;-----------------------------------------------------------------------------------------------------------------
 
@@ -59,21 +61,23 @@ SuggestedDriverVersion="24.3.1"
 ; The newest driver which supports legacy AMD GPUs (GCN1-GCN3 chips like the Radeon R9 or older) is 22.6.1. We allow this driver for DX11.
 +DriverDenyList=(DriverVersion="<27.20.20913.2000", Reason="These driver versions have known stability issues and missing features", RHI="D3D11")
 
-; The minimum version for DX12 is Vulkan is 24.3.1, published in March 2024. There are separate packages for different GPU generations and they
+; The minimum version for DX12 is 24.3.1, published in March 2024. There are separate packages for different GPU generations and they
 ; have different internal versions, so we can't check for a version number explicitly (the RDNA driver is 31.0.24027.1012, older GPUs have
 ; 31.0.21912.14). We'll use a driver date check instead, as both packages have internal date stamps newer than Feb 19th.
 +DriverDenyList=(DriverDate="<2-19-2024", Reason="These driver versions have known stability issues and missing features", RHI="D3D12")
-+DriverDenyList=(DriverDate="<2-19-2024", Reason="These driver versions have known stability issues and missing features", RHI="Vulkan")
+
+; The minimum version for Vulkan is 25.2.1, published in Feb 2025.
++DriverDenyList=(DriverDate="<2-10-2025", Reason="These driver versions have known stability issues and missing features", RHI="Vulkan")
 
 ;-----------------------------------------------------------------------------------------------------------------
 
 [GPU_AMD Linux]
-SuggestedDriverVersion="23.3.1"
+SuggestedDriverVersion="24.2.8"
 
-; RADV builds older than 23.3.1 have issues running UE5, so we'll use this as the minimum version (released December 2023). Note that it's also
-; possible to use the AMD proprietary driver on Linux, and we don't distinguish between driver vendors. The current version for the AMD driver as
-; of writing this is 23.40.2, so 23.3.1 is a valid version for that package as well (although a fairly old one).
-+DriverDenyList=(DriverVersion="<23.3.1", Reason="These driver versions have known stability issues and missing features")
+; RADV builds older than 24.2.8 have issues running UE5, so we'll use this as the minimum version.
+; For SM6 and ray tracing, verison 25.0.0+ of RADV is recommended.
+; Note that it's also possible to use the AMD proprietary driver on Linux, and we don't distinguish between driver vendors.
++DriverDenyList=(DriverVersion="<24.2.8", Reason="These driver versions have known stability issues and missing features")
 
 ;-----------------------------------------------------------------------------------------------------------------
 
@@ -91,4 +95,7 @@ SuggestedDriverVersion="101.5333"
 ;-----------------------------------------------------------------------------------------------------------------
 
 [GPU_Intel Linux]
-SuggestedDriverVersion="30.0.101.1340"
+SuggestedDriverVersion="24.2.8"
+
+; Mesa builds older than 24.2.8 have issues running UE5, so we'll use this as the minimum version.
++DriverDenyList=(DriverVersion="<24.2.8", Reason="These driver versions have known stability issues and missing features")
diff --git a/Engine/Source/Runtime/RHI/Private/Windows/WindowsDynamicRHI.cpp b/Engine/Source/Runtime/RHI/Private/Windows/WindowsDynamicRHI.cpp
--- a/Engine/Source/Runtime/RHI/Private/Windows/WindowsDynamicRHI.cpp
+++ b/Engine/Source/Runtime/RHI/Private/Windows/WindowsDynamicRHI.cpp
@@ -987,7 +987,7 @@ static bool HandleUnsupportedRHI(EWindowsRHI& WindowsRHI, ERHIFeatureLevel::Type
 
 	if (WindowsRHI == EWindowsRHI::Vulkan)
 	{
-		FMessageDialog::Open(EAppMsgType::Ok, LOCTEXT("RequiredVulkan", "Vulkan Driver is required to run the engine."));
+		FMessageDialog::Open(EAppMsgType::Ok, LOCTEXT("RequiredVulkan", "A compatible Vulkan Driver is required to run the engine."));
 		FPlatformMisc::RequestExit(true, TEXT("HandleUnsupportedRHI.Vulkan"));
 	}
 
