Added support for Google Play Games for PC Emulator

#android

diff --git a/Engine/Build/Android/Java/src/com/epicgames/unreal/GameActivity.java.template b/Engine/Build/Android/Java/src/com/epicgames/unreal/GameActivity.java.template
--- a/Engine/Build/Android/Java/src/com/epicgames/unreal/GameActivity.java.template
+++ b/Engine/Build/Android/Java/src/com/epicgames/unreal/GameActivity.java.template
@@ -2435,21 +2435,19 @@ Trace.endSection();
 		// check for a processor id
 		boolean bIsARM64 = false;
 		boolean bHasNEON = false;
-		String processor = "";
-		if (CpuInfo.containsKey("Processor"))
-		{
-			processor = CpuInfo.get("Processor");
-		}
-		if (processor.contains("aarch64"))
+		String processor = CpuInfo.getOrDefault("Processor", "");
+        String abilist = getProp("ro.product.cpu.abilist");
+		
+		if (processor.contains("aarch64") || abilist.contains("arm64-v8a"))
 		{
 			bIsARM64 = true;
 			bHasNEON = true;
 		}
-
-		String abilist = getProp("ro.product.cpu.abilist");
-		if (abilist.contains("arm64-v8a"))
+		else if (getPackageManager().hasSystemFeature("com.google.android.play.feature.HPE_EXPERIENCE"))
 		{
 			bIsARM64 = true;
+			bHasNEON = true;
+			bHasHoudini = false;
 		}
 
 		// check each processor featureset
diff --git a/Engine/Config/BaseDeviceProfiles.ini b/Engine/Config/BaseDeviceProfiles.ini
--- a/Engine/Config/BaseDeviceProfiles.ini
+++ b/Engine/Config/BaseDeviceProfiles.ini
@@ -109,6 +109,7 @@
 +DeviceProfileNameAndTypes=Android_Mid,Android
 +DeviceProfileNameAndTypes=Android_High,Android
 +DeviceProfileNameAndTypes=Android_Emulator,Android
++DeviceProfileNameAndTypes=Android_PC_Emulator,Android
 +DeviceProfileNameAndTypes=Android_Default,Android
 +DeviceProfileNameAndTypes=Android_Adreno4xx,Android
 +DeviceProfileNameAndTypes=Android_Adreno5xx_Low,Android
@@ -792,8 +793,11 @@ BaseProfileName=VisionPro
 +MatchProfile=(Profile="Oculus_Quest",Match=((SourceType=SRC_DeviceMake,CompareType=CMP_Equal,MatchString="Oculus")))
 +MatchProfile=(Profile="MagicLeap_Vulkan",Match=((SourceType=SRC_DeviceMake,CompareType=CMP_Equal,MatchString="Magic Leap"), (SourceType=SRC_VulkanAvailable,CompareType=CMP_Equal,MatchString="true")))
 
-; Match Android emulator
-+MatchProfile=(Profile="Android_Emulator",Match=((SourceType=SRC_DeviceMake,CompareType=CMP_Equal,MatchString="Google"),(SourceType=SRC_GPUFamily,CompareType=CMP_Regex,MatchString="Android Emulator")))
+; Match Android Emulator
++MatchProfile=(Profile="Android_Emulator",Match=((SourceType=SRC_DeviceMake,CompareType=CMP_Equal,MatchString="Google"),(SourceType=SRC_GpuFamily,CompareType=CMP_Regex,MatchString="Android Emulator")))
+
+; Match PC Emulator
++MatchProfile=(Profile="Android_PC_Emulator",Match=((SourceType=SRC_DeviceMake,CompareType=CMP_Equal,MatchString="Google"),(SourceType=SRC_DeviceModel,CompareType=CMP_Equal,MatchString="HPE device")))
 
 ; Adreno-based devices
 +MatchProfile=(Profile="Android_Adreno4xx",Match=((SourceType=SRC_GpuFamily,CompareType=CMP_Regex,MatchString="Adreno \\(TM\\) 4[0-9][0-9]")))
@@ -954,6 +958,12 @@ DeviceType=Android
 BaseProfileName=Android_High
 +CVars=r.Android.DisableVulkanSupport=0
 
+[Android_PC_Emulator DeviceProfile]
+DeviceType=Android
+BaseProfileName=Android_Emulator
++CVars=r.Vulkan.SupportsBCTextureFormats=1
++CVars=android.Vulkan.HasUnifiedMemory=0
+
 [Android_Default DeviceProfile]
 DeviceType=Android
 BaseProfileName=Android_Mid
diff --git a/Engine/Source/Runtime/VulkanRHI/Private/Android/VulkanAndroidPlatform.cpp b/Engine/Source/Runtime/VulkanRHI/Private/Android/VulkanAndroidPlatform.cpp
--- a/Engine/Source/Runtime/VulkanRHI/Private/Android/VulkanAndroidPlatform.cpp
+++ b/Engine/Source/Runtime/VulkanRHI/Private/Android/VulkanAndroidPlatform.cpp
@@ -73,6 +73,17 @@ static TAutoConsoleVariable<int32> CVarVulkanSupportsBCTextureFormats(
 	TEXT("  1 = supported."),
 	ECVF_SetByDeviceProfile
 );
+#if PLATFORM_ANDROID_X64 	
+
+static TAutoConsoleVariable<int32> CVarVulkanHasUnifiedMemory(
+	TEXT("android.Vulkan.HasUnifiedMemory"),
+	1,
+	TEXT("Whether or not device has unified memory\n")
+	TEXT("  0 = doesn't have\n")
+	TEXT("  1 = has."),
+	ECVF_SetByDeviceProfile
+);
+#endif
 
 // Vulkan function pointers
 #define DEFINE_VK_ENTRYPOINTS(Type,Func) Type VulkanDynamicAPI::Func = NULL;
@@ -838,6 +849,13 @@ void FVulkanAndroidPlatform::OverridePlatformHandlers(bool bInit)
 		FPlatformMisc::SetOnPauseCallback(nullptr);
 	}
 }
+#if PLATFORM_ANDROID_X64
+
+bool FVulkanAndroidPlatform::HasUnifiedMemory()
+{
+	return CVarVulkanHasUnifiedMemory.GetValueOnAnyThread();
+}
+#endif
 
 void FVulkanAndroidPlatform::SetupRequiresDepthStencilFullWriteWorkaround(const FVulkanDevice& Device)
 {
diff --git a/Engine/Source/Runtime/VulkanRHI/Private/Android/VulkanAndroidPlatform.h b/Engine/Source/Runtime/VulkanRHI/Private/Android/VulkanAndroidPlatform.h
--- a/Engine/Source/Runtime/VulkanRHI/Private/Android/VulkanAndroidPlatform.h
+++ b/Engine/Source/Runtime/VulkanRHI/Private/Android/VulkanAndroidPlatform.h
@@ -182,7 +182,11 @@ public:
 	//#todo-rco: Detect Mali?
 	static bool RequiresPresentLayoutFix() { return true; }
 
+#if PLATFORM_ANDROID_X64 	
+	static bool HasUnifiedMemory();
+#else
 	static bool HasUnifiedMemory() { return true; }
+#endif
 
 	static bool RegisterGPUWork() { return false; }
 
