Adding enum redirect for ECollectionScriptingShareType

[RN]

diff --git a/Engine/Config/BaseEngine.ini b/Engine/Config/BaseEngine.ini
--- a/Engine/Config/BaseEngine.ini
+++ b/Engine/Config/BaseEngine.ini
@@ -1354,6 +1354,7 @@ ClassName=/Script/OnlineSubsystemUtils.OnlineEngineInterfaceImpl
 +ClassRedirects=(OldName="/Script/HierarchyEditor.HierarchyCategory",NewName="/Script/DataHierarchyEditor.HierarchyCategory")
 
 +EnumRedirects=(OldName="EMaterialExpressionConvertType", ValueChanges=(("Float1", "Scalar"), ("Float2", "Vector2"), ("Float3", "Vector3"), ("Float4", "Vector4")))
++EnumRedirects=(OldName="/Script/AssetTags.ECollectionScriptingShareType", NewName="/Script/Engine.ECollectionScriptingShareType")
 
 +StructRedirects=(OldName="/Script/MassEntity.MassProcessingContext",NewName="/Script/MassEntity.MassProcessingContext_DEPRECATED")
 ; End of releases redirects
