[Modular Control Rig] - Added Module tags and Root56 module to the correct ini files.

diff --git a/Engine/Config/BaseEditor.ini b/Engine/Config/BaseEditor.ini
--- a/Engine/Config/BaseEditor.ini
+++ b/Engine/Config/BaseEditor.ini
@@ -715,4 +715,9 @@ bUseWidgetTemplateSelector=True
 +Filters=(FilterString="FColorVertexBuffer",DisplayName="Color Vertex Buffers")
 
 [MemoryPools]
-PersistentAllocatorReserveSizeMB=2048
\ No newline at end of file
+PersistentAllocatorReserveSizeMB=2048
+
+[/Script/RigVM.RigVMProjectSettings]
++VariantTags=(Name="ControlRigModules_Epic",Label="Epic",ToolTip=NSLOCTEXT("[/Script/RigVM]", "561F4A5D4EE5ED7AFFE366BFA988185A", "Epic built modules"),Color=(R=0.514918,G=0.046665,B=1.000000,A=1.000000),bShowInUserInterface=True,bMarksSubjectAsInvalid=False)
++VariantTags=(Name="ControlRigModules_Biped",Label="Biped",ToolTip=NSLOCTEXT("[/Script/RigVM]", "1F5E3123499AAC91FEA5B084B539EAE8", "Biped modules"),Color=(R=0.514918,G=0.046665,B=1.000000,A=1.000000),bShowInUserInterface=True,bMarksSubjectAsInvalid=False)
++VariantTags=(Name="ControlRigModules_Quadruped",Label="Quadruped",ToolTip=NSLOCTEXT("[/Script/RigVM]", "1A58BF794BCA84FB89CD74B4C75152B7", "Quadruped modules"),Color=(R=0.514918,G=0.046665,B=1.000000,A=1.000000),bShowInUserInterface=True,bMarksSubjectAsInvalid=False)
diff --git a/Engine/Config/BaseEngine.ini b/Engine/Config/BaseEngine.ini
--- a/Engine/Config/BaseEngine.ini
+++ b/Engine/Config/BaseEngine.ini
@@ -3879,3 +3879,6 @@ RuntimeHashDefaultClass=/Script/CoreUObject.Class'/Script/Engine.WorldPartitionR
 
 [BuildPatchTool]
 VerificationThreadCount=4
+
+[/Script/ControlRig.ControlRigSettings]
+DefaultRootModule=/ControlRig/Modules/Modules56/Root.Root
\ No newline at end of file
