Enable ProjectLauncher extensions by default

#rnx

diff --git a/Engine/Config/BaseEditorPerProjectUserSettings.ini b/Engine/Config/BaseEditorPerProjectUserSettings.ini
--- a/Engine/Config/BaseEditorPerProjectUserSettings.ini
+++ b/Engine/Config/BaseEditorPerProjectUserSettings.ini
@@ -156,6 +156,7 @@ bToolbarCustomization=False
 bBehaviorTreeEditor=False
 bEnableFindAndReplaceReferences =False
 bExampleLayersAndBlends = True
+bEnableProjectLauncherExtensions=True
 
 ; Enabled by default now
 bEnableAsyncTextureCompilation=True
