Add Low/Mid/High Android device profiles for PowerVR ABCD Series GPUs based on core counts
Fix crash with CheckAndroidDeviceProfile commandlet missing resolution and inset values
#rnx

diff --git a/Engine/Config/BaseDeviceProfiles.ini b/Engine/Config/BaseDeviceProfiles.ini
--- a/Engine/Config/BaseDeviceProfiles.ini
+++ b/Engine/Config/BaseDeviceProfiles.ini
@@ -152,6 +152,12 @@
 +DeviceProfileNameAndTypes=Android_PowerVR_GE8xxx,Android
 +DeviceProfileNameAndTypes=Android_PowerVR_GM9xxx,Android
 +DeviceProfileNameAndTypes=Android_PowerVR_GM9xxx_Vulkan,Android
++DeviceProfileNameAndTypes=Android_PowerVR_Series_Low,Android
++DeviceProfileNameAndTypes=Android_PowerVR_Series_Mid,Android
++DeviceProfileNameAndTypes=Android_PowerVR_Series_High,Android
++DeviceProfileNameAndTypes=Android_PowerVR_Series_Low_Vulkan,Android
++DeviceProfileNameAndTypes=Android_PowerVR_Series_Mid_Vulkan,Android
++DeviceProfileNameAndTypes=Android_PowerVR_Series_High_Vulkan,Android
 +DeviceProfileNameAndTypes=Android_TegraK1,Android
 +DeviceProfileNameAndTypes=Android_Unknown_Vulkan,Android
 +DeviceProfileNameAndTypes=Meta_Quest_3,Android
@@ -859,10 +865,14 @@ BaseProfileName=VisionPro
 +MatchProfile=(Profile="Android_PowerVR_GE8xxx",Match=((SourceType=SRC_GpuFamily,CompareType=CMP_Regex,MatchString="PowerVR Rogue GE8[0-9]+")))
 +MatchProfile=(Profile="Android_PowerVR_GM9xxx_Vulkan",Match=((SourceType=SRC_GpuFamily,CompareType=CMP_Regex,MatchString="PowerVR Rogue GM9[0-9]+"),(SourceType=SRC_VulkanAvailable,CompareType=CMP_Equal,MatchString="true")))
 +MatchProfile=(Profile="Android_PowerVR_GM9xxx",Match=((SourceType=SRC_GpuFamily,CompareType=CMP_Regex,MatchString="PowerVR Rogue GM9[0-9]+")))
-+MatchProfile=(Profile="Android_PowerVR_ABCD_Vulkan",Match=((SourceType=SRC_GpuFamily,CompareType=CMP_Regex,MatchString="PowerVR [A-D]-Series"),(SourceType=SRC_VulkanAvailable,CompareType=CMP_Equal,MatchString="true")))
-+MatchProfile=(Profile="Android_PowerVR_ABCD",Match=((SourceType=SRC_GpuFamily,CompareType=CMP_Regex,MatchString="PowerVR [A-D]-Series")))
-; Tegra devices
++MatchProfile=(Profile="Android_PowerVR_Series_High_Vulkan",Match=((SourceType=SRC_GpuFamily,CompareType=CMP_Regex,MatchString="PowerVR [A-D]\\-Series [A-D]X[A-Z]+\\-([0-9]+)\\-[0-9]+"),(SourceType=SRC_PreviousRegexMatch,CompareType=CMP_GreaterEqual,MatchString="32"),(SourceType=SRC_VulkanAvailable,CompareType=CMP_Equal,MatchString="true")))
++MatchProfile=(Profile="Android_PowerVR_Series_Mid_Vulkan",Match=((SourceType=SRC_GpuFamily,CompareType=CMP_Regex,MatchString="PowerVR [A-D]\\-Series [A-D]X[A-Z]+\\-([0-9]+)\\-[0-9]+"),(SourceType=SRC_PreviousRegexMatch,CompareType=CMP_GreaterEqual,MatchString="16"),(SourceType=SRC_VulkanAvailable,CompareType=CMP_Equal,MatchString="true")))
++MatchProfile=(Profile="Android_PowerVR_Series_Low_Vulkan",Match=((SourceType=SRC_GpuFamily,CompareType=CMP_Regex,MatchString="PowerVR [A-D]\\-Series [A-D]X[A-Z]+\\-[0-9]+\\-[0-9]+"),(SourceType=SRC_VulkanAvailable,CompareType=CMP_Equal,MatchString="true")))
++MatchProfile=(Profile="Android_PowerVR_Series_High",Match=((SourceType=SRC_GpuFamily,CompareType=CMP_Regex,MatchString="PowerVR [A-D]\\-Series [A-D]X[A-Z]+\\-([0-9]+)\\-[0-9]+"),(SourceType=SRC_PreviousRegexMatch,CompareType=CMP_GreaterEqual,MatchString="32")))
++MatchProfile=(Profile="Android_PowerVR_Series_Mid",Match=((SourceType=SRC_GpuFamily,CompareType=CMP_Regex,MatchString="PowerVR [A-D]\\-Series [A-D]X[A-Z]+\\-([0-9]+)\\-[0-9]+"),(SourceType=SRC_PreviousRegexMatch,CompareType=CMP_GreaterEqual,MatchString="16")))
++MatchProfile=(Profile="Android_PowerVR_Series_Low",Match=((SourceType=SRC_GpuFamily,CompareType=CMP_Regex,MatchString="PowerVR [A-D]\\-Series [A-D]X[A-Z]+\\-[0-9]+\\-[0-9]+")))
 
+; Tegra devices
 +MatchProfile=(Profile="Android_TegraK1",Match=((SourceType=SRC_GpuFamily,CompareType=CMP_Equal,MatchString="NVIDIA Tegra"),(SourceType=SRC_GlVersion,CompareType=CMP_Regex,MatchString="^OpenGL ES 3\\.")))
 
 ; Android_Default is the profile of last resort
@@ -1138,13 +1148,32 @@ BaseProfileName=Android_PowerVR_GM9xxx
 ; enable Vulkan
 +CVars=r.Android.DisableVulkanSupport=0
 
-[Android_PowerVR_ABCD DeviceProfile]
+[Android_PowerVR_Series_High DeviceProfile]
 DeviceType=Android
-BaseProfileName=Android_PowerVR_GM9xxx
+BaseProfileName=Android_High
 
-[Android_PowerVR_ABCD_Vulkan DeviceProfile]
+[Android_PowerVR_Series_Mid DeviceProfile]
 DeviceType=Android
-BaseProfileName=Android_PowerVR_GM9xxx_Vulkan
+BaseProfileName=Android_Mid
+
+[Android_PowerVR_Series_Low DeviceProfile]
+DeviceType=Android
+BaseProfileName=Android_Low
+
+[Android_PowerVR_Series_High_Vulkan DeviceProfile]
+DeviceType=Android
+BaseProfileName=Android_PowerVR_Series_High
++CVars=r.Android.DisableVulkanSupport=0
+
+[Android_PowerVR_Series_Mid_Vulkan DeviceProfile]
+DeviceType=Android
+BaseProfileName=Android_PowerVR_Series_Mid
++CVars=r.Android.DisableVulkanSupport=0
+
+[Android_PowerVR_Series_Low_Vulkan DeviceProfile]
+DeviceType=Android
+BaseProfileName=Android_PowerVR_Series_Low
++CVars=r.Android.DisableVulkanSupport=0
 
 ;----------------
 ; Mali devices, we support from Mali-T6xx and above
diff --git a/Engine/Source/Editor/DeviceProfileServices/Private/CheckAndroidDeviceProfileCommandlet.cpp b/Engine/Source/Editor/DeviceProfileServices/Private/CheckAndroidDeviceProfileCommandlet.cpp
--- a/Engine/Source/Editor/DeviceProfileServices/Private/CheckAndroidDeviceProfileCommandlet.cpp
+++ b/Engine/Source/Editor/DeviceProfileServices/Private/CheckAndroidDeviceProfileCommandlet.cpp
@@ -150,6 +150,12 @@ int32 UCheckAndroidDeviceProfileCommandlet::Main(const FString& RawCommandLine)
 				DeviceParameters.Add(FName(TEXT("SRC_TotalPhysicalGB")), Params.FindRef(TEXT("TotalPhysicalGB")));
 				DeviceParameters.Add(FName(TEXT("SRC_HMDSystemName")), TEXT(""));
 				DeviceParameters.Add(FName(TEXT("SRC_SM5Available")), Params.FindRef(TEXT("SM5Available")));
+				DeviceParameters.Add(FName(TEXT("SRC_ResolutionX")), Params.FindRef(TEXT("ResolutionX")));
+				DeviceParameters.Add(FName(TEXT("SRC_ResolutionY")), Params.FindRef(TEXT("ResolutionY")));
+				DeviceParameters.Add(FName(TEXT("SRC_InsetsLeft")), Params.FindRef(TEXT("InsetsLeft")));
+				DeviceParameters.Add(FName(TEXT("SRC_InsetsTop")), Params.FindRef(TEXT("InsetsTop")));
+				DeviceParameters.Add(FName(TEXT("SRC_InsetsRight")), Params.FindRef(TEXT("InsetsRight")));
+				DeviceParameters.Add(FName(TEXT("SRC_InsetsBottom")), Params.FindRef(TEXT("InsetsBottom")));
 			}
 			AndroidDeviceProfileSelector->SetSelectorProperties(DeviceParameters);
 			FString ProfileName = ForcedDeviceProfileName.IsEmpty() ? AndroidDeviceProfileSelector->GetDeviceProfileName() : ForcedDeviceProfileName;
