Enable android port forwarder by default

diff --git a/Engine/Config/BaseEngine.ini b/Engine/Config/BaseEngine.ini
--- a/Engine/Config/BaseEngine.ini
+++ b/Engine/Config/BaseEngine.ini
@@ -2614,6 +2614,9 @@ RuntimeTargetPlatform="{Platform}"
 [Zen]
 AutoLaunch=true
 
+[Zen.Plugin.AndroidPortForwarder]
+Name=Engine/Source/Developer/Android/AndroidZenServerPlugin/Binaries/Win64/AndroidPortForwarder.dll
+
 [Zen.AutoLaunch]
 DataPath=%ENGINEVERSIONAGNOSTICINSTALLEDUSERDIR%Zen/Data
 LimitProcessLifetime=true
@@ -2621,6 +2624,7 @@ AllowRemoteNetworkService=false
 LocalDataCachePathEnvOverride=UE-LocalDataCachePath
 LocalDataCachePathCommandLineOverride=LocalDataCachePath
 LocalDataCachePathEditorOverrideSetting=LocalDerivedDataCache
++Plugins=AndroidPortForwarder
 ExtraArgs=--http asio --gc-cache-duration-seconds 1209600 --gc-interval-seconds 21600 --gc-low-diskspace-threshold 2147483648 --quiet
 
 [Zen.ConnectExisting]
