Normalize a few default High scalability VSM settings across platforms.

diff --git a/Engine/Config/BaseScalability.ini b/Engine/Config/BaseScalability.ini
--- a/Engine/Config/BaseScalability.ini
+++ b/Engine/Config/BaseScalability.ini
@@ -195,13 +195,13 @@ r.VolumetricFog.HistoryMissSupersampleCount=4
 r.LightMaxDrawDistanceScale=1
 r.CapsuleShadows=1
 r.Shadow.Virtual.MaxPhysicalPages=2048
-r.Shadow.Virtual.ResolutionLodBiasDirectional=-0.5
-r.Shadow.Virtual.ResolutionLodBiasDirectionalMoving=-0.5
+r.Shadow.Virtual.ResolutionLodBiasDirectional=0.0
+r.Shadow.Virtual.ResolutionLodBiasDirectionalMoving=0.0
 r.Shadow.Virtual.ResolutionLodBiasLocal=0.0
 r.Shadow.Virtual.ResolutionLodBiasLocalMoving=1.0
 r.Shadow.Virtual.SMRT.RayCountDirectional=8
 r.Shadow.Virtual.SMRT.SamplesPerRayDirectional=4
-r.Shadow.Virtual.SMRT.RayCountLocal=8
+r.Shadow.Virtual.SMRT.RayCountLocal=4
 r.Shadow.Virtual.SMRT.SamplesPerRayLocal=4
 
 [ShadowQuality@3]
