Fix scalability issue where AOQuality wasn't being properly initialized

diff --git a/Engine/Config/BaseScalability.ini b/Engine/Config/BaseScalability.ini
--- a/Engine/Config/BaseScalability.ini
+++ b/Engine/Config/BaseScalability.ini
@@ -262,6 +262,7 @@ r.Shadow.Virtual.SMRT.SamplesPerRayLocal=8
 
 [GlobalIlluminationQuality@0]
 ; When DFAO is disabled, Skylight intensity is suppressed
+r.AOQuality=0
 r.DistanceFieldAO=0
 r.SkylightIntensityMultiplier=0.8
 r.Lumen.DiffuseIndirect.Allow=0
