Heterogeneous Volumes: updates to profile settings

diff --git a/Engine/Config/BaseScalability.ini b/Engine/Config/BaseScalability.ini
--- a/Engine/Config/BaseScalability.ini
+++ b/Engine/Config/BaseScalability.ini
@@ -727,8 +727,8 @@ r.SkyAtmosphere.TransmittanceLUT.UseSmallFormat=0
 r.SkyAtmosphere.TransmittanceLUT.SampleCount=10.0
 r.SkyAtmosphere.MultiScatteringLUT.SampleCount=15.0
 fx.Niagara.QualityLevel=2
-r.HeterogeneousVolumes.DownsampleFactor=3
-r.HeterogeneousVolumes.MaxStepCount=256
+r.HeterogeneousVolumes.DownsampleFactor=2
+r.HeterogeneousVolumes.MaxStepCount=96
 r.HeterogeneousVolumes.Shadows.Resolution=256
 
 [EffectsQuality@3]
