Set SubtitleWidgetToUse default to be consistent with underlying fallback of SubtitleWidgetToUseDefault

During testing/reviewing we noticed the default for SubtitleWidgetToUse is empty. In this case, the right things happen as the underlying code falls back to SubtitleWidgetToUseDefault. Still, setting the default SubtitleWidgetToUse visibly that way in the first place is more intuitive for the user and low risk as it's using the same value.

Also, the subtitle config was originally set to "Engine" when it was restricted. Now that it has been promoted to experimental, set it to "Game" to support game-specific subtitle display parameter defaults.

diff --git a/Engine/Config/BaseEngine.ini b/Engine/Config/BaseEngine.ini
--- a/Engine/Config/BaseEngine.ini
+++ b/Engine/Config/BaseEngine.ini
@@ -3886,6 +3886,3 @@ VerificationThreadCount=4
 
 [/Script/ControlRig.ControlRigSettings]
 DefaultRootModule=/ControlRig/Modules/Modules56/Root.Root
-
-[/Script/SubtitlesAndClosedCaptions.SubtitlesSettings]
-SubtitleWidgetToUseDefault=/Script/UMG.WidgetBlueprintGeneratedClass'/SubtitlesAndClosedCaptions/Widgets/DefaultSubtitleWidget.DefaultSubtitleWidget_C'
\ No newline at end of file
diff --git a/Engine/Config/BaseGame.ini b/Engine/Config/BaseGame.ini
--- a/Engine/Config/BaseGame.ini
+++ b/Engine/Config/BaseGame.ini
@@ -228,6 +228,10 @@ QueryCountWarningInterval=30.0
 FrameInterpolationProcessor=Class'/Script/LiveLink.LiveLinkBasicFrameInterpolationProcessor'
 +DefaultRoleSettings=(Role=Class'/Script/LiveLink.LiveLinkAnimationRole', FrameInterpolationProcessor=Class'/Script/LiveLink.LiveLinkAnimationFrameInterpolationProcessor')
 
+[/Script/SubtitlesAndClosedCaptions.SubtitlesSettings]
+SubtitleWidgetToUse=/Script/UMG.WidgetBlueprintGeneratedClass'/SubtitlesAndClosedCaptions/Widgets/DefaultSubtitleWidget.DefaultSubtitleWidget_C'
+SubtitleWidgetToUseDefault=/Script/UMG.WidgetBlueprintGeneratedClass'/SubtitlesAndClosedCaptions/Widgets/DefaultSubtitleWidget.DefaultSubtitleWidget_C'
+
 [/Script/Engine.AssetManagerSettings]
 @PrimaryAssetTypesToScan=PrimaryAssetType
 
diff --git a/Engine/Plugins/Experimental/SubtitlesAndClosedCaptions/Source/SubtitlesAndClosedCaptions/Public/SubtitlesSettings.h b/Engine/Plugins/Experimental/SubtitlesAndClosedCaptions/Source/SubtitlesAndClosedCaptions/Public/SubtitlesSettings.h
--- a/Engine/Plugins/Experimental/SubtitlesAndClosedCaptions/Source/SubtitlesAndClosedCaptions/Public/SubtitlesSettings.h
+++ b/Engine/Plugins/Experimental/SubtitlesAndClosedCaptions/Source/SubtitlesAndClosedCaptions/Public/SubtitlesSettings.h
@@ -8,7 +8,7 @@
 #include "SubtitlesSettings.generated.h"
 
 
-UCLASS(config = Engine, defaultconfig, meta = (DisplayName = "Subtitles And Closed Captions"))
+UCLASS(config = Game, defaultconfig, meta = (DisplayName = "Subtitles And Closed Captions"))
 class USubtitlesSettings : public UDeveloperSettings
 {
 	GENERATED_BODY()
