Adds commented out cvars for the new translator to make it easier for developers to enable similar to the shader cvars in this file.

r.Material.Translator.EnableNew=1
r.Material.Translator.DebugDump=2

#rnx

diff --git a/Engine/Config/ConsoleVariables.ini b/Engine/Config/ConsoleVariables.ini
--- a/Engine/Config/ConsoleVariables.ini
+++ b/Engine/Config/ConsoleVariables.ini
@@ -57,6 +57,11 @@ r.DumpShaderDebugInfo=2
 ; When this is enabled, shader compiler warnings are emitted to the log for all shaders as they are loaded (either from the DDC or from a shader compilation job).
 ;r.ShaderCompiler.EmitWarningsOnLoad=1
 
+; When this is enabled it will enable the new translator
+;r.Material.Translator.EnableNew=1
+; When this is enabled the new translator will dump out the IR graph to a dotviz graph to your project's saved directory.
+;r.Material.Translator.DebugDump=2
+
 ; Uncomment to disable parallel rendering
 ;r.RHICmdBypass=1
 ; Uncomment to disable parallel mesh draw command setup
