Expand the amount of default monitor resolutions for PIE and UMG Designer to include common modern options

Also renamed existing entries to use the resolution designation rather than physical monitor size.

To see the new options, users will have to clear out the saved `MonitorScreenResolutions` entries under `[/Script/UnrealEd.LevelEditorPlaySettings]` in `[Project]/Saved/Config/WindowsEditor/EditorPerProjectUserSettings.ini`

#rnx

diff --git a/Engine/Config/BaseEditorPerProjectUserSettings.ini b/Engine/Config/BaseEditorPerProjectUserSettings.ini
--- a/Engine/Config/BaseEditorPerProjectUserSettings.ini
+++ b/Engine/Config/BaseEditorPerProjectUserSettings.ini
@@ -254,11 +254,20 @@ LastExecutedLaunchPlatform=
 +LaptopScreenResolutions=(Description="Generic 14-15.6\" Notebook",Width=1366,Height=768,AspectRatio="16:9",bCanSwapAspectRatio=true)
 
 ; common screen resolutions for desktop monitors
-+MonitorScreenResolutions=(Description="19\" monitor",Width=1440,Height=900,AspectRatio="16:10",bCanSwapAspectRatio=true)
-+MonitorScreenResolutions=(Description="20\" monitor",Width=1600,Height=900,AspectRatio="16:9",bCanSwapAspectRatio=true)
-+MonitorScreenResolutions=(Description="22\" monitor",Width=1680,Height=1050,AspectRatio="16:10",bCanSwapAspectRatio=true)
-+MonitorScreenResolutions=(Description="21.5-24\" monitor",Width=1920,Height=1080,AspectRatio="16:9",bCanSwapAspectRatio=true)
-+MonitorScreenResolutions=(Description="27\" monitor",Width=2560,Height=1440,AspectRatio="16:9",bCanSwapAspectRatio=true)
++MonitorScreenResolutions=(Description="16:9 HD",Width=1280,Height=720,AspectRatio="16:9",bCanSwapAspectRatio=true)
++MonitorScreenResolutions=(Description="16:10 WXGA+",Width=1440,Height=900,AspectRatio="16:10",bCanSwapAspectRatio=true)
++MonitorScreenResolutions=(Description="16:9 HD+",Width=1600,Height=900,AspectRatio="16:9",bCanSwapAspectRatio=true)
++MonitorScreenResolutions=(Description="16:10 WSXGA+",Width=1680,Height=1050,AspectRatio="16:10",bCanSwapAspectRatio=true)
++MonitorScreenResolutions=(Description="16:9 Full HD",Width=1920,Height=1080,AspectRatio="16:9",bCanSwapAspectRatio=true)
++MonitorScreenResolutions=(Description="16:10 WUXGA",Width=1920,Height=1200,AspectRatio="16:10",bCanSwapAspectRatio=true)
++MonitorScreenResolutions=(Description="16:9 2K",Width=2560,Height=1440,AspectRatio="16:9",bCanSwapAspectRatio=true)
++MonitorScreenResolutions=(Description="16:9 4K",Width=3840,Height=2160,AspectRatio="16:9",bCanSwapAspectRatio=true)
++MonitorScreenResolutions=(Description="21:9 Ultrawide FHD",Width=2520,Height=1080,AspectRatio="21:9",bCanSwapAspectRatio=true)
++MonitorScreenResolutions=(Description="21:9 Ultrawide QHD",Width=3440,Height=1440,AspectRatio="21:9",bCanSwapAspectRatio=true)
++MonitorScreenResolutions=(Description="32:9 DFHD",Width=3840,Height=1080,AspectRatio="32:9",bCanSwapAspectRatio=true)
++MonitorScreenResolutions=(Description="32:9 DQHD",Width=5120,Height=1440,AspectRatio="32:9",bCanSwapAspectRatio=true)
++MonitorScreenResolutions=(Description="4:3 XGA",Width=1024,Height=768,AspectRatio="4:3",bCanSwapAspectRatio=true)
++MonitorScreenResolutions=(Description="3:2 WQHD+",Width=2160,Height=1440,AspectRatio="3:2",bCanSwapAspectRatio=true)
 
 ; common screen resolutions for mobile phones
 +PhoneScreenResolutions=(Description="Apple iPhone 5S",Width=320,Height=568,AspectRatio="~16:9",bCanSwapAspectRatio=true,ProfileName="iPhone5S")
