LiveLink -  Animation roles will now correctly use the animation interpolation by default

diff --git a/Engine/Config/BaseGame.ini b/Engine/Config/BaseGame.ini
--- a/Engine/Config/BaseGame.ini
+++ b/Engine/Config/BaseGame.ini
@@ -227,7 +227,7 @@ QueryCountWarningInterval=30.0
 
 [/Script/LiveLink.LiveLinkSettings]
 FrameInterpolationProcessor=Class'/Script/LiveLink.LiveLinkBasicFrameInterpolationProcessor'
-+DefaultRoleSettings=(Role=Class'/Script/LiveLink.LiveLinkAnimationRole', FrameInterpolationProcessor=Class'/Script/LiveLink.LiveLinkAnimationFrameInterpolationProcessor')
+DefaultRoleSettings=(Role="/Script/CoreUObject.Class'/Script/LiveLinkInterface.LiveLinkAnimationRole'",SettingClass="/Script/CoreUObject.Class'/Script/LiveLinkInterface.LiveLinkSubjectSettings'",FrameInterpolationProcessor="/Script/CoreUObject.Class'/Script/LiveLink.LiveLinkAnimationFrameInterpolationProcessor'",FramePreProcessors=)
 
 [/Script/SubtitlesAndClosedCaptions.SubtitlesSettings]
 SubtitleWidgetToUse=/Script/UMG.WidgetBlueprintGeneratedClass'/SubtitlesAndClosedCaptions/Widgets/DefaultSubtitleWidget.DefaultSubtitleWidget_C'
