Change invalid r.MaxAnisotropy=0 to r.MaxAnisotropy=1

We have a Clamp in ComputeAnisotropyRT anyway, but just make things more clear.

diff --git a/Engine/Config/BaseScalability.ini b/Engine/Config/BaseScalability.ini
--- a/Engine/Config/BaseScalability.ini
+++ b/Engine/Config/BaseScalability.ini
@@ -596,7 +596,7 @@ r.Streaming.MipBias=16
 r.Streaming.AmortizeCPUToGPUCopy=1
 r.Streaming.MaxNumTexturesToStreamPerFrame=1
 r.Streaming.Boost=0.3
-r.MaxAnisotropy=0
+r.MaxAnisotropy=1
 r.VT.MaxAnisotropy=4
 r.Streaming.LimitPoolSizeToVRAM=1
 r.Streaming.PoolSize=400
