NetEmulation
* Added a "BufferBloat" emulation profile that enables buffer bloat of 400ms (in and out) over an 'average' emulation profile

diff --git a/Engine/Config/BaseEngine.ini b/Engine/Config/BaseEngine.ini
--- a/Engine/Config/BaseEngine.ini
+++ b/Engine/Config/BaseEngine.ini
@@ -3540,9 +3540,21 @@ PktLagMax=200
 PktIncomingLagMin=100
 PktIncomingLagMax=200
 
+[PacketSimulationProfile.BufferBloat]
+; Simulates average network conditions but with buffer bloat stalling traffic for 400ms
+PktLoss=1
+PktIncomingLoss=1
+PktLagMin=30
+PktLagMax=60
+PktIncomingLagMin=30
+PktIncomingLagMax=60
+PktBufferBloatInMS=400
+PktIncomingBufferBloatInMS=400
+
 [/Script/Engine.NetworkSettings]
 +NetworkEmulationProfiles=(ProfileName="Average",ToolTip="Simulates average internet conditions")
 +NetworkEmulationProfiles=(ProfileName="Bad",ToolTip="Simulates laggy internet conditions")
++NetworkEmulationProfiles=(ProfileName="BufferBloat",ToolTip="Simulates a connection affected by buffer bloat (in and out)")
 
 [PacketHandlerComponents]
 ; Options can be set in this section of DefaultEngine.ini to enable different types of network packet encruption plugins
