Enable overwrite prevention on Local DDC.  This is available in zenserver 5.6.16 and later.  This will not affect Shared DDC.  After monitoring for any issues, we can consider making this configuration change for internal Shared DDC servers.

diff --git a/Engine/Config/BaseEngine.ini b/Engine/Config/BaseEngine.ini
--- a/Engine/Config/BaseEngine.ini
+++ b/Engine/Config/BaseEngine.ini
@@ -2655,7 +2655,7 @@ AllowRemoteNetworkService=false
 LocalDataCachePathEnvOverride=UE-LocalDataCachePath
 LocalDataCachePathCommandLineOverride=LocalDataCachePath
 LocalDataCachePathEditorOverrideSetting=LocalDerivedDataCache
-ExtraArgs=--http asio --gc-cache-duration-seconds 1209600 --gc-interval-seconds 21600 --gc-low-diskspace-threshold 2147483648 --quiet
+ExtraArgs=--http asio --gc-cache-duration-seconds 1209600 --gc-interval-seconds 21600 --gc-low-diskspace-threshold 2147483648 --cache-bucket-limit-overwrites --quiet
 
 [Zen.ConnectExisting]
 HostName=localhost
