New Material Translator:  Support for 13 more expressions, plus support for the "use preview as default" flag on material function inputs.  Bug fixed in TrySimplifyOperator for (0 - X), was returning 0 instead of -X.  Bug fixed for transposed matrix multiply.  Matrix * vector should interpret the second argument as a column vector, and interpret the output column vector result back as a row vector (HLSL doesn't differentiate row or column vectors, but the translator does).  Fixed natvis for Composite vector types.  Fixed null dereference in TextureSample expression, when a TextureObject input is provided.  Negate operator implementation added for non-constant code path.

#rnx

diff --git a/Engine/Config/BaseMaterialExpressions.ini b/Engine/Config/BaseMaterialExpressions.ini
--- a/Engine/Config/BaseMaterialExpressions.ini
+++ b/Engine/Config/BaseMaterialExpressions.ini
@@ -116,8 +116,21 @@ bShaderInputData=true
 [/Script/Engine.MaterialExpressionLegacyBlendMaterialAttributes]
 +MenuCategories=NSLOCTEXT("MaterialExpression", "MaterialAttributes", "Material Attributes")
 
+; Order of external code identifiers matches EMaterialExpressionBoundsType (set of 4 identifiers per enum, InstanceLocal, ObjectLocal, PreSkinnedLocal)
 [/Script/Engine.MaterialExpressionBounds]
 +MenuCategories=NSLOCTEXT("MaterialExpression", "Vectors", "Vectors")
++ExternalCodeIdentifiers="InstanceLocalBoundsHalf"
++ExternalCodeIdentifiers="InstanceLocalBoundsFull"
++ExternalCodeIdentifiers="InstanceLocalBoundsMin"
++ExternalCodeIdentifiers="InstanceLocalBoundsMax"
++ExternalCodeIdentifiers="ObjectLocalBoundsHalf"
++ExternalCodeIdentifiers="ObjectLocalBoundsFull"
++ExternalCodeIdentifiers="ObjectLocalBoundsMin"
++ExternalCodeIdentifiers="ObjectLocalBoundsMax"
++ExternalCodeIdentifiers="PreSkinnedLocalBoundsHalf"
++ExternalCodeIdentifiers="PreSkinnedLocalBoundsFull"
++ExternalCodeIdentifiers="PreSkinnedLocalBoundsMin"
++ExternalCodeIdentifiers="PreSkinnedLocalBoundsMax"
 bShaderInputData=true
 
 [/Script/Engine.MaterialExpressionBreakMaterialAttributes]
@@ -480,6 +493,7 @@ bShaderInputData=true
 
 [/Script/Engine.MaterialExpressionScalarBlueNoise]
 +MenuCategories=NSLOCTEXT("MaterialExpressionNoise", "Noise", "Noise")
++ExternalCodeIdentifiers="ScalarBlueNoise"
 
 [/Script/Engine.MaterialExpressionNormalize]
 +MenuCategories=NSLOCTEXT("MaterialExpression", "Math", "Math")
@@ -492,11 +506,16 @@ bShaderInputData=true
 
 [/Script/Engine.MaterialExpressionObjectLocalBounds]
 +MenuCategories=NSLOCTEXT("MaterialExpression", "Vectors", "Vectors")
++ExternalCodeIdentifiers="ObjectLocalBoundsHalf"
++ExternalCodeIdentifiers="ObjectLocalBoundsFull"
++ExternalCodeIdentifiers="ObjectLocalBoundsMin"
++ExternalCodeIdentifiers="ObjectLocalBoundsMax"
 bShaderInputData=true
 
 [/Script/Engine.MaterialExpressionObjectOrientation]
 +MenuCategories=NSLOCTEXT("MaterialExpression", "Vectors", "Vectors")
 +MenuCategories=NSLOCTEXT("MaterialExpression", "Coordinates", "Coordinates")
++ExternalCodeIdentifiers="ObjectOrientation"
 bShaderInputData=true
 
 [/Script/Engine.MaterialExpressionObjectPositionWS]
@@ -542,6 +561,8 @@ bShaderInputData=true
 +MenuCategories=NSLOCTEXT("MaterialExpression", "Particles", "Particles")
 +MenuCategories=NSLOCTEXT("MaterialExpression", "Coordinates", "Coordinates")
 +MenuCategories=NSLOCTEXT("MaterialExpression", "Vectors", "Vectors")
++ExternalCodeIdentifiers="ParticlePositionWSAbsolute"
++ExternalCodeIdentifiers="ParticlePositionWSTranslated"
 bShaderInputData=true
 
 [/Script/Engine.MaterialExpressionParticleRadius]
@@ -604,10 +625,12 @@ bShaderInputData=true
 
 [/Script/Engine.MaterialExpressionPerInstanceCustomData]
 +MenuCategories=NSLOCTEXT("MaterialExpression", "Custom", "Custom")
++ExternalCodeIdentifiers="GetPerInstanceCustomData"
 bShaderInputData=true
 
 [/Script/Engine.MaterialExpressionPerInstanceCustomData3Vector]
 +MenuCategories=NSLOCTEXT("MaterialExpression", "Custom", "Custom")
++ExternalCodeIdentifiers="GetPerInstanceCustomData3Vector"
 bShaderInputData=true
 
 [/Script/Engine.MaterialExpressionPerInstanceFadeAmount]
@@ -641,6 +664,7 @@ bShaderInputData=true
 
 [/Script/Engine.MaterialExpressionPrecomputedAOMask]
 +MenuCategories=NSLOCTEXT("MaterialExpression", "Constants", "Constants")
++ExternalCodeIdentifiers="PrecomputedAOMask"
 bShaderInputData=true
 
 [/Script/Engine.MaterialExpressionPreSkinnedLocalBounds]
@@ -1025,11 +1049,14 @@ bShaderInputData=true
 +ExternalCodeDeclarations=(ReturnType=Float,			Name="MeshPaintTextureCoordinateIndex",	Definition="GetMeshPaintTextureCoordinateIndex(GetPrimitiveData(Parameters))", bIsInlined=true)
 +ExternalCodeDeclarations=(ReturnType=TextureMeshPaint,	Name="MeshPaintTextureDescriptor",		Definition="GetMeshPaintTextureDescriptor(GetPrimitiveData(Parameters))", Derivative=Zero, bIsInlined=true)
 +ExternalCodeDeclarations=(ReturnType=Float3,			Name="ObjectBounds",					Definition="float3(GetPrimitiveData(Parameters).ObjectBoundsX, GetPrimitiveData(Parameters).ObjectBoundsY, GetPrimitiveData(Parameters).ObjectBoundsZ)", Derivative=Zero, bIsInlined=true, Domains=(MD_Surface, MD_Volume))
++ExternalCodeDeclarations=(ReturnType=Float3,			Name="ObjectOrientation",				Definition="GetObjectOrientation(Parameters)", Derivative=Zero, Domains=(MD_Surface, MD_Volume))
 +ExternalCodeDeclarations=(ReturnType=Float,			Name="ObjectRadius",					Definition="GetPrimitiveData(Parameters).ObjectRadius", Derivative=Zero, bIsInlined=true, Domains=(MD_Surface, MD_Volume))
 +ExternalCodeDeclarations=(ReturnType=Float4,			Name="ParticleColor",					Definition="Parameters.Particle.Color", Derivative=Zero, bIsInlined=true, EnvironmentDefines=((Name="NEEDS_PARTICLE_COLOR", ShaderFrequency=Pixel|Compute)))
 +ExternalCodeDeclarations=(ReturnType=Float3,			Name="ParticleDirection",				Definition="Parameters.Particle.Velocity.xyz", Derivative=Zero, bIsInlined=true, EnvironmentDefines=((Name="NEEDS_PARTICLE_VELOCITY")))
 +ExternalCodeDeclarations=(ReturnType=Float2,			Name="ParticleMacroUV",					Definition="GetParticleMacroUV(Parameters)", ShaderFrequency=Pixel|Compute)
 +ExternalCodeDeclarations=(ReturnType=Float,			Name="ParticleMotionBlurFade",			Definition="Parameters.Particle.MotionBlurFade", Derivative=Zero, bIsInlined=true, EnvironmentDefines=((Name="USES_PARTICLE_MOTION_BLUR")))
++ExternalCodeDeclarations=(ReturnType=LWCVector3,		Name="ParticlePositionWSAbsolute",		Definition="Get<PREV>ParticleWorldPosition(Parameters)", Derivative=Zero, EnvironmentDefines=((Name="NEEDS_PARTICLE_POSITION")))
++ExternalCodeDeclarations=(ReturnType=Float3,			Name="ParticlePositionWSTranslated",	Definition="Parameters.Particle.<PREV>TranslatedWorldPositionAndSize.xyz", Derivative=Zero, EnvironmentDefines=((Name="NEEDS_PARTICLE_POSITION")))
 +ExternalCodeDeclarations=(ReturnType=Float,			Name="ParticleRadius",					Definition="max(Parameters.Particle.TranslatedWorldPositionAndSize.w, .001f)", Derivative=Zero, bIsInlined=true, EnvironmentDefines=((Name="NEEDS_PARTICLE_POSITION")))
 +ExternalCodeDeclarations=(ReturnType=Float,			Name="ParticleRelativeTime",			Definition="Parameters.Particle.RelativeTime", Derivative=Zero, bIsInlined=true, EnvironmentDefines=((Name="NEEDS_PARTICLE_TIME")))
 +ExternalCodeDeclarations=(ReturnType=Float,			Name="ParticleRandom",					Definition="Parameters.Particle.Random", Derivative=Zero, bIsInlined=true, EnvironmentDefines=((Name="NEEDS_PARTICLE_RANDOM")))
@@ -1041,9 +1068,11 @@ bShaderInputData=true
 +ExternalCodeDeclarations=(ReturnType=Float,			Name="ParticleSubUVLerp",				Definition="Parameters.Particle.SubUVLerp", EnvironmentDefines=((Name="USE_PARTICLE_SUBUVS")))
 +ExternalCodeDeclarations=(ReturnType=Float1,			Name="PixelDepth",						Definition="GetPixelDepth(Parameters)", DefinitionDDX="Parameters.ScreenPosition_DDX.w", DefinitionDDY="Parameters.ScreenPosition_DDY.w", Derivative=Valid, bIsInlined=true)
 +ExternalCodeDeclarations=(ReturnType=Float2,			Name="PixelPosition",					Definition="GetPixelPosition(Parameters)", DefinitionDDX="float2(1.0f, 0.0f)", DefinitionDDY="float2(0.0f, 1.0f)", Derivative=Valid)
++ExternalCodeDeclarations=(ReturnType=Float,			Name="PrecomputedAOMask",				Definition="Parameters.AOMaterialMask", Derivative=Zero, EnvironmentDefines=(("USES_AO_MATERIAL_MASK")), ShaderFrequency=Pixel|Compute, MinimumFeatureLevel=SM5)
 +ExternalCodeDeclarations=(ReturnType=Float3,			Name="PreSkinnedPosition",				Definition="Parameters.PreSkinnedPosition", ShaderFrequency=Vertex, bIsInlined=true)
 +ExternalCodeDeclarations=(ReturnType=Float3,			Name="PreSkinnedNormal",				Definition="Parameters.PreSkinnedNormal", ShaderFrequency=Vertex, bIsInlined=true)
 +ExternalCodeDeclarations=(ReturnType=Float3,			Name="ReflectionVector",				Definition="Parameters.ReflectionVector", ShaderFrequency=Pixel|Compute, bIsInlined=true, EnvironmentDefines=((Name="USES_TRANSFORM_VECTOR", ShaderFrequency=Pixel|Compute)))
++ExternalCodeDeclarations=(ReturnType=Float,			Name="ScalarBlueNoise",					Definition="ViewScalarBlueNoise(GetPixelPosition(Parameters), View.StateFrameIndex)", ShaderFrequency=Pixel|Compute)
 +ExternalCodeDeclarations=(ReturnType=Float,			Name="SphericalParticleOpacity",		Definition="GetSphericalParticleOpacity(Parameters,$0)", ShaderFrequency=Pixel|Compute, EnvironmentDefines=((Name="NEEDS_PARTICLE_POSITION"), (Name="SPHERICAL_PARTICLE_OPACITY"), (Name="NEEDS_WORLD_POSITION_EXCLUDING_SHADER_OFFSETS"), (Name="MIR.SceneDepth")))
 +ExternalCodeDeclarations=(ReturnType=Float,			Name="TwoSidedSign",					Definition="Parameters.TwoSidedSign", ShaderFrequency=Pixel|Compute, Derivative=Zero, bIsInlined=true)
 +ExternalCodeDeclarations=(ReturnType=Float4,			Name="VertexColor",						Definition="Parameters.VertexColor", DefinitionDDX="Parameters.VertexColor_DDX", DefinitionDDY="Parameters.VertexColor_DDY", Derivative=Valid, bIsInlined=true, EnvironmentDefines=((Name="INTERPOLATE_VERTEX_COLOR", ShaderFrequency=Pixel|Compute)))
@@ -1051,8 +1080,10 @@ bShaderInputData=true
 +ExternalCodeDeclarations=(ReturnType=Float3,			Name="VertexTangent",					Definition="Parameters.TangentToWorld[0]", bIsInlined=true, EnvironmentDefines=((Name="USES_TRANSFORM_VECTOR", ShaderFrequency=Pixel|Compute)))
 +ExternalCodeDeclarations=(ReturnType=Float2,			Name="ViewportUV",						Definition="GetViewportUV(Parameters)", DefinitionDDX="float2(View.ViewSizeAndInvSize.z, 0.0f)", DefinitionDDY="float2(0.0f, View.ViewSizeAndInvSize.w)", Derivative=Valid)
 
-+ExternalCodeDeclarations=(ReturnType=Float,			Name="GetPerInstanceRandom",			Definition="GetPerInstanceRandom(Parameters)", ShaderFrequency=Vertex|Pixel, Derivative=Zero, bIsInlined=true, EnvironmentDefines=((Name="NEEDS_PER_INSTANCE_RANDOM_PS", ShaderFrequency=Pixel)))
++ExternalCodeDeclarations=(ReturnType=Float,			Name="GetPerInstanceCustomData",		Definition="GetPerInstanceCustomData(Parameters, $0, $1)", Derivative=Zero, EnvironmentDefines=((Name="USES_PER_INSTANCE_CUSTOM_DATA")))
++ExternalCodeDeclarations=(ReturnType=Float3,			Name="GetPerInstanceCustomData3Vector",	Definition="GetPerInstanceCustomData3Vector(Parameters, $0, $1)", Derivative=Zero, EnvironmentDefines=((Name="USES_PER_INSTANCE_CUSTOM_DATA")))
 +ExternalCodeDeclarations=(ReturnType=Float,			Name="GetPerInstanceFadeAmount",		Definition="GetPerInstanceFadeAmount(Parameters)", ShaderFrequency=Vertex|Pixel, Derivative=Zero, bIsInlined=true, EnvironmentDefines=((Name="USES_PER_INSTANCE_FADE_AMOUNT")))
++ExternalCodeDeclarations=(ReturnType=Float,			Name="GetPerInstanceRandom",			Definition="GetPerInstanceRandom(Parameters)", ShaderFrequency=Vertex|Pixel, Derivative=Zero, bIsInlined=true, EnvironmentDefines=((Name="NEEDS_PER_INSTANCE_RANDOM_PS", ShaderFrequency=Pixel)))
 
 +ExternalCodeDeclarations=(ReturnType=Float,	Name="GetCloudSampleAltitude",								Definition="MaterialExpressionCloudSampleAltitude(Parameters)")
 +ExternalCodeDeclarations=(ReturnType=Float,	Name="GetCloudSampleAltitudeInLayer",						Definition="MaterialExpressionCloudSampleAltitudeInLayer(Parameters)")
@@ -1077,6 +1108,19 @@ bShaderInputData=true
 +ExternalCodeDeclarations=(ReturnType=Bool,		Name="ReflectionCapturePassSwitch",					Definition="GetReflectionCapturePassSwitchState()")
 +ExternalCodeDeclarations=(ReturnType=Bool,		Name="ShadowReplace",								Definition="GetShadowReplaceState()")
 
++ExternalCodeDeclarations=(ReturnType=Float3,	Name="InstanceLocalBoundsHalf",		Definition="(GetPrimitiveData(Parameters).InstanceLocalBoundsExtent)", Domains=(MD_Surface, MD_Volume))
++ExternalCodeDeclarations=(ReturnType=Float3,	Name="InstanceLocalBoundsFull",		Definition="(GetPrimitiveData(Parameters).InstanceLocalBoundsExtent * 2.0f)", Domains=(MD_Surface, MD_Volume))
++ExternalCodeDeclarations=(ReturnType=Float3,	Name="InstanceLocalBoundsMin",		Definition="(GetPrimitiveData(Parameters).InstanceLocalBoundsCenter - GetPrimitiveData(Parameters).InstanceLocalBoundsExtent)", Domains=(MD_Surface, MD_Volume))
++ExternalCodeDeclarations=(ReturnType=Float3,	Name="InstanceLocalBoundsMax",		Definition="(GetPrimitiveData(Parameters).InstanceLocalBoundsCenter + GetPrimitiveData(Parameters).InstanceLocalBoundsExtent)", Domains=(MD_Surface, MD_Volume))
++ExternalCodeDeclarations=(ReturnType=Float3,	Name="ObjectLocalBoundsHalf",		Definition="((GetPrimitiveData(Parameters).LocalObjectBoundsMax - GetPrimitiveData(Parameters).LocalObjectBoundsMin) / 2.0f)", Domains=(MD_Surface, MD_Volume))
++ExternalCodeDeclarations=(ReturnType=Float3,	Name="ObjectLocalBoundsFull",		Definition="(GetPrimitiveData(Parameters).LocalObjectBoundsMax - GetPrimitiveData(Parameters).LocalObjectBoundsMin)", Domains=(MD_Surface, MD_Volume))
++ExternalCodeDeclarations=(ReturnType=Float3,	Name="ObjectLocalBoundsMin",		Definition="GetPrimitiveData(Parameters).LocalObjectBoundsMin", Domains=(MD_Surface, MD_Volume))
++ExternalCodeDeclarations=(ReturnType=Float3,	Name="ObjectLocalBoundsMax",		Definition="GetPrimitiveData(Parameters).LocalObjectBoundsMax", Domains=(MD_Surface, MD_Volume))
++ExternalCodeDeclarations=(ReturnType=Float3,	Name="PreSkinnedLocalBoundsHalf",	Definition="((GetPrimitiveData(Parameters).PreSkinnedLocalBoundsMax - GetPrimitiveData(Parameters).PreSkinnedLocalBoundsMin) / 2.0f)", Domains=(MD_Surface, MD_Volume))
++ExternalCodeDeclarations=(ReturnType=Float3,	Name="PreSkinnedLocalBoundsFull",	Definition="(GetPrimitiveData(Parameters).PreSkinnedLocalBoundsMax - GetPrimitiveData(Parameters).PreSkinnedLocalBoundsMin)", Domains=(MD_Surface, MD_Volume))
++ExternalCodeDeclarations=(ReturnType=Float3,	Name="PreSkinnedLocalBoundsMin",	Definition="GetPrimitiveData(Parameters).PreSkinnedLocalBoundsMin", Domains=(MD_Surface, MD_Volume))
++ExternalCodeDeclarations=(ReturnType=Float3,	Name="PreSkinnedLocalBoundsMax",	Definition="GetPrimitiveData(Parameters).PreSkinnedLocalBoundsMax", Domains=(MD_Surface, MD_Volume))
+
 ; External code declarations for Transform and TransformPosition nodes
 +ExternalCodeDeclarations=(ReturnType=Float4x4,		Name="CameraViewToTranslatedWorld", Definition="ResolvedView.<PREV>CameraViewToTranslatedWorld", Derivative=Zero, bIsInlined=true)
 +ExternalCodeDeclarations=(ReturnType=LWCMatrix,	Name="GetInstanceToWorld",			Definition="Get<PREV>InstanceToWorld(Parameters)", Derivative=Zero, bIsInlined=true, EnvironmentDefines=((Name="NEEDS_INSTANCE_LOCAL_TO_WORLD_PS", ShaderFrequency=Pixel)))
diff --git a/Engine/Source/Runtime/Engine/Private/Materials/MaterialExpressionsIR.cpp b/Engine/Source/Runtime/Engine/Private/Materials/MaterialExpressionsIR.cpp
--- a/Engine/Source/Runtime/Engine/Private/Materials/MaterialExpressionsIR.cpp
+++ b/Engine/Source/Runtime/Engine/Private/Materials/MaterialExpressionsIR.cpp
@@ -910,7 +910,7 @@ void UMaterialExpressionSquareRoot::Build(MIR::FEmitter& Em)
 	Em.Output(0, Em.Sqrt(Em.Input(&Input)));
 }
 
-static FValueRef EmitInlineHLSL(MIR::FEmitter& Em, const UMaterialExpressionExternalCodeBase& InExternalCodeExpression, int32 InExternalCodeIdentifierIndex, TConstArrayView<FValueRef> InArguments)
+static FValueRef EmitInlineHLSL(MIR::FEmitter& Em, const UMaterialExpressionExternalCodeBase& InExternalCodeExpression, int32 InExternalCodeIdentifierIndex, TConstArrayView<FValueRef> InArguments, MIR::EValueFlags ValueFlags = MIR::EValueFlags::None)
 {
 	checkf(InExternalCodeIdentifierIndex >= 0 && InExternalCodeIdentifierIndex < InExternalCodeExpression.ExternalCodeIdentifiers.Num(),
 		TEXT("External code identifier index (%d) out of bounds; Upper bound is %d"), InExternalCodeIdentifierIndex, InExternalCodeExpression.ExternalCodeIdentifiers.Num());
@@ -923,7 +923,7 @@ static FValueRef EmitInlineHLSL(MIR::FEmitter& Em, const UMaterialExpressionExte
 		return {};
 	}
 
-	return Em.InlineHLSL(ExternalCodeDeclaration, InArguments);
+	return Em.InlineHLSL(ExternalCodeDeclaration, InArguments, ValueFlags);
 }
 
 static void BuildInlineHLSLOutput(MIR::FEmitter& Em, const UMaterialExpressionExternalCodeBase& InExternalCodeExpression, TConstArrayView<FValueRef> InArguments)
@@ -1366,7 +1366,7 @@ void UMaterialExpressionTextureSample::Build(MIR::FEmitter& Em)
 
 	UE_MIR_CHECKPOINT(Em);
 
-	BuildTextureSampleExpression(Em, this, TextureValue, Texture->GetMaterialType());
+	BuildTextureSampleExpression(Em, this, TextureValue, TextureValue->GetTexture()->GetMaterialType());
 }
 
 static FValueRef BuildTextureObjectParameter(MIR::FEmitter& Em, UMaterialExpressionTextureSampleParameter* Expr)
@@ -3426,4 +3426,192 @@ void UMaterialExpressionCustom::Build(MIR::FEmitter& Em)
 	}
 }
 
+void UMaterialExpressionBounds::Build(MIR::FEmitter& Em)
+{
+	// Select between 3 different sets of 4 outputs (half, full, min, max), depending on bounds type.  Check that enum matches order in BaseMaterialExpressions.ini.
+	static_assert(MEILB_InstanceLocal == 0);
+	static_assert(MEILB_ObjectLocal == 1);
+	static_assert(MEILB_PreSkinnedLocal == 2);
+
+	int32 OutputOffset = Type * 4;
+	Em.Output(0, EmitInlineHLSL(Em, *this, OutputOffset + 0, {}));
+	Em.Output(1, EmitInlineHLSL(Em, *this, OutputOffset + 1, {}));
+	Em.Output(2, EmitInlineHLSL(Em, *this, OutputOffset + 2, {}));
+	Em.Output(3, EmitInlineHLSL(Em, *this, OutputOffset + 3, {}));
+}
+
+void UMaterialExpressionBumpOffset::Build(MIR::FEmitter& Em)
+{
+	FValueRef HeightRatioArg = Em.Cast(Em.InputDefaultFloat(&HeightRatioInput, HeightRatio), MIR::FType::MakeFloatScalar());
+
+	FValueRef TexCoordArg = Em.TryInput(&Coordinate);
+	if (!TexCoordArg)
+	{
+		TexCoordArg = Em.ExternalInput(MIR::TexCoordIndexToExternalInput(ConstCoordinate));
+	}
+
+	Em.Output(0,
+		Em.Add(
+			Em.Multiply(
+				Em.Swizzle(EmitTransformVectorBase(Em, EmitInlineHLSL(Em, NAME_CameraVector, {}), MCB_World, MCB_Tangent, false, nullptr, nullptr), MIR::FSwizzleMask(MIR::EVectorComponent::X, MIR::EVectorComponent::Y)),
+				Em.Add(
+					Em.Multiply(
+						HeightRatioArg,
+						Em.Cast(Em.Input(&Height), MIR::FType::MakeFloatScalar())
+					),
+					Em.Multiply(Em.ConstantFloat(-ReferencePlane), HeightRatioArg)
+				)
+			),
+			TexCoordArg
+		)
+	);
+}
+
+void UMaterialExpressionDynamicParameter::Build(MIR::FEmitter& Em)
+{
+	FValueRef Default = Em.ConstantFloat4(FVector4f(DefaultValue.R, DefaultValue.G, DefaultValue.B, DefaultValue.A));
+	FValueRef DynamicParameterIndex = Em.ExternalInput(MIR::EExternalInput::DynamicParticleParameterIndex, ParameterIndex);
+	FValueRef Result = Em.InlineHLSL(MIR::FType::MakeFloatVector(4), TEXTVIEW("GetDynamicParameter(Parameters.Particle, $0, $1)"), { Default, DynamicParameterIndex });
+
+	Em.Output(0, Em.Subscript(Result, 0));
+	Em.Output(1, Em.Subscript(Result, 1));
+	Em.Output(2, Em.Subscript(Result, 2));
+	Em.Output(3, Em.Subscript(Result, 3));
+	Em.Output(4, Em.Swizzle(Result, MIR::FSwizzleMask::XYZ()));		// RGB
+	Em.Output(5, Result);											// RGBA
+}
+
+void UMaterialExpressionNoise::Build(MIR::FEmitter& Em)
+{
+	FValueRef PositionInput = EmitWorldPositionOrDefault(Em, Em.TryInput(&Position), WorldPositionOriginType);
+
+	if (WorldPositionOriginType == EPositionOrigin::CameraRelative)
+	{
+		// LWC_TODO: add support for translated world positions in the corresponding HLSL function
+		PositionInput = EmitTransformVectorBase(Em, PositionInput, MCB_TranslatedWorld, MCB_World, true, nullptr, nullptr);
+	}
+
+	FValueRef FilterWidthInput = Em.InputDefaultFloat(&FilterWidth, 0.0f);
+	FValueRef ScaleValue = Em.ConstantFloat(Scale);
+	FValueRef QualityValue = Em.ConstantInt(Quality);
+	FValueRef NoiseFunctionValue = Em.ConstantInt(NoiseFunction);
+	FValueRef TurbulenceValue = Em.ConstantBool(bTurbulence);
+	FValueRef LevelsValue = Em.ConstantInt(FMath::Clamp(Levels, 1, 10));		// to limit performance problems due to values outside reasonable range
+	FValueRef OutputMinValue = Em.ConstantFloat(OutputMin);
+	FValueRef OutputMaxValue = Em.ConstantFloat(OutputMax);
+	FValueRef LevelScaleValue = Em.ConstantFloat(LevelScale);
+	FValueRef TilingValue = Em.ConstantBool(bTiling);
+	FValueRef RepeatSizeValue = Em.ConstantFloat(RepeatSize);
+
+	Em.Output(0, Em.InlineHLSL(MIR::FType::MakeFloatScalar(),
+		TEXTVIEW("MaterialExpressionNoise($0,$1,$2,$3,$4,$5,$6,$7,$8,$9,$10,$11)"), {
+			PositionInput,
+			ScaleValue,
+			QualityValue,
+			NoiseFunctionValue,
+			TurbulenceValue,
+			LevelsValue,
+			OutputMinValue,
+			OutputMaxValue,
+			LevelScaleValue,
+			FilterWidthInput,
+			TilingValue,
+			RepeatSizeValue
+		}));
+}
+
+void UMaterialExpressionVectorNoise::Build(MIR::FEmitter& Em)
+{
+	FValueRef PositionInput = EmitWorldPositionOrDefault(Em, Em.TryInput(&Position), WorldPositionOriginType);
+
+	if (WorldPositionOriginType == EPositionOrigin::CameraRelative)
+	{
+		// LWC_TODO: add support for translated world positions in the corresponding HLSL function
+		PositionInput = EmitTransformVectorBase(Em, PositionInput, MCB_TranslatedWorld, MCB_World, true, nullptr, nullptr);
+	}
+
+	// LWC_TODO - maybe possible/useful to add LWC-aware noise functions
+	PositionInput = Em.Cast(PositionInput, MIR::FType::MakeFloatVector(3));
+
+	FValueRef QualityValue = Em.ConstantInt(Quality);
+	FValueRef NoiseFunctionValue = Em.ConstantInt(NoiseFunction);
+	FValueRef TilingValue = Em.ConstantBool(bTiling);
+	FValueRef TileSizeValue = Em.ConstantFloat(TileSize);
+
+	FValueRef NoiseResult = Em.InlineHLSL(MIR::FType::MakeFloatVector(4),
+		TEXTVIEW("MaterialExpressionVectorNoise($0,$1,$2,$3,$4)"), {
+			PositionInput,
+			QualityValue,
+			NoiseFunctionValue,
+			TilingValue,
+			TileSizeValue
+		});
+
+	// Function returns float4, but only certain noise functions fill in all four elements, so downcast to float3 if not those cases.
+	if (NoiseFunction != VNF_GradientALU && NoiseFunction != VNF_VoronoiALU)
+	{
+		NoiseResult = Em.Cast(NoiseResult, MIR::FType::MakeFloatVector(3));
+	}
+
+	Em.Output(0, NoiseResult);
+}
+
+void UMaterialExpressionPanner::Build(MIR::FEmitter& Em)
+{
+	FValueRef TimeArg = Em.TryInput(&Time);
+	if (TimeArg)
+	{
+		TimeArg = Em.Cast(TimeArg, MIR::FType::MakeFloatScalar());
+	}
+	else
+	{
+		TimeArg = Em.InlineHLSL(MIR::FType::MakeFloatScalar(), TEXTVIEW("View.<PREV>GameTime"), {}, MIR::EValueFlags::SubstituteTags);
+	}
+	
+	FValueRef SpeedVectorArg = Em.InputDefaultFloat2(&Speed, FVector2f(SpeedX, SpeedY));
+
+	// TODO:  When preshaders get implemented, the original translator generates a unique "PeriodicHint" preshader op for this expression,
+	//        which attempts to do math at higher precision to avoid accuracy issues as GameTime increases.  We'll want to add that logic here,
+	//        or consider making preshader math involving game time automatically run at high precision across the board (naturally solving
+	//        precision issues even outside this specific expression).
+
+	SpeedVectorArg = Em.Multiply(TimeArg, SpeedVectorArg);
+	if (bFractionalPart)
+	{
+		SpeedVectorArg = Em.Frac(SpeedVectorArg);
+	}
+
+	FValueRef TexCoordArg = Em.TryInput(&Coordinate);
+	if (!TexCoordArg)
+	{
+		TexCoordArg = Em.ExternalInput(MIR::TexCoordIndexToExternalInput(0));
+	}
+
+	Em.Output(0, Em.Add(SpeedVectorArg, TexCoordArg));
+}
+
+void UMaterialExpressionParticlePositionWS::Build(MIR::FEmitter& Em)
+{
+	int32 ExternalCodeIndex = OriginType == EPositionOrigin::Absolute ? 0 : 1;
+	MIR::FType ResultType = OriginType == EPositionOrigin::Absolute ? MIR::FType::MakeDoubleVector(3) : MIR::FType::MakeFloatVector(3);
+
+	Em.Output(0, EmitInlineHLSL(Em, *this, ExternalCodeIndex, {}, MIR::EValueFlags::SubstituteTags));
+}
+
+void UMaterialExpressionPerInstanceCustomData::Build(MIR::FEmitter& Em)
+{
+	FValueRef DataIndexArgument = Em.ConstantInt(DataIndex);
+	FValueRef DefaultArgument = Em.InputDefaultFloat(&DefaultValue, ConstDefaultValue);
+
+	Em.Output(0, EmitInlineHLSL(Em, *this, 0, { DataIndexArgument, DefaultArgument }));
+}
+
+void UMaterialExpressionPerInstanceCustomData3Vector::Build(MIR::FEmitter& Em)
+{
+	FValueRef DataIndexArgument = Em.ConstantInt(DataIndex);
+	FValueRef DefaultArgument = Em.InputDefaultFloat3(&DefaultValue, FVector3f(ConstDefaultValue.R, ConstDefaultValue.G, ConstDefaultValue.B));
+
+	Em.Output(0, EmitInlineHLSL(Em, *this, 0, { DataIndexArgument, DefaultArgument }));
+}
+
 #endif // WITH_EDITOR
diff --git a/Engine/Source/Runtime/Engine/Private/Materials/MaterialIR.cpp b/Engine/Source/Runtime/Engine/Private/Materials/MaterialIR.cpp
--- a/Engine/Source/Runtime/Engine/Private/Materials/MaterialIR.cpp
+++ b/Engine/Source/Runtime/Engine/Private/Materials/MaterialIR.cpp
@@ -538,6 +538,7 @@ const TCHAR* LexToString(EExternalInput Input)
 		case EExternalInput::ViewMaterialTextureMipBias: return TEXT("ViewMaterialTextureMipBias");
 		case EExternalInput::ViewMaterialTextureDerivativeMultiply: return TEXT("ViewMaterialTextureDerivativeMultiply");
 		case EExternalInput::GlobalDistanceField: return TEXT("GlobalDistanceField");
+		case EExternalInput::DynamicParticleParameterIndex: return TEXT("DynamicParticleParameterIndex");
 		default: UE_MIR_UNREACHABLE();
 	}
 }
@@ -576,6 +577,9 @@ FType GetExternalInputType(EExternalInput Id)
 		case EExternalInput::WorldPosition_CameraRelativeNoOffsets_Ddy:
 			return FType::MakeFloatVector(3);
 
+		case EExternalInput::DynamicParticleParameterIndex:
+			return FType::MakeIntScalar();
+
 		default: UE_MIR_UNREACHABLE();
 	}
 }
diff --git a/Engine/Source/Runtime/Engine/Private/Materials/MaterialIRDebug.cpp b/Engine/Source/Runtime/Engine/Private/Materials/MaterialIRDebug.cpp
--- a/Engine/Source/Runtime/Engine/Private/Materials/MaterialIRDebug.cpp
+++ b/Engine/Source/Runtime/Engine/Private/Materials/MaterialIRDebug.cpp
@@ -123,6 +123,10 @@ struct FDebugDumpIRUseGraphState
 		else if (const FExternalInput* ExternalInput = Value->As<FExternalInput>())
 		{
 			Out.Append(LexToString(ExternalInput->Id));
+			if (ExternalInput->Id == EExternalInput::DynamicParticleParameterIndex)
+			{
+				Out.Appendf(TEXT("[%u]"), ExternalInput->UserData);
+			}
 		}
 		else if (const FSetMaterialOutput* SetMaterailOutput = Value->As<FSetMaterialOutput>())
 		{
diff --git a/Engine/Source/Runtime/Engine/Private/Materials/MaterialIREmitter.cpp b/Engine/Source/Runtime/Engine/Private/Materials/MaterialIREmitter.cpp
--- a/Engine/Source/Runtime/Engine/Private/Materials/MaterialIREmitter.cpp
+++ b/Engine/Source/Runtime/Engine/Private/Materials/MaterialIREmitter.cpp
@@ -939,10 +939,11 @@ FValueRef FEmitter::Poison()
 	return FPoison::Get();
 }
 
-FValueRef FEmitter::ExternalInput(EExternalInput Id)
+FValueRef FEmitter::ExternalInput(EExternalInput Id, uint32 UserData)
 {
 	FExternalInput Prototype = MakePrototype<FExternalInput>(GetExternalInputType(Id));
 	Prototype.Id = Id;
+	Prototype.UserData = UserData;
 	return EmitPrototype(*this, Prototype);
 }
 
@@ -1650,7 +1651,7 @@ static FValueRef TrySimplifyOperator(FEmitter& Emitter, EOperator& Op, FValueRef
 			}
 			else if (A->AreAllNearlyZero())
 			{
-				return Emitter.Negate(A);
+				return Emitter.Negate(B);
 			}
 			else if (A->Equals(B))
 			{
@@ -2323,16 +2324,32 @@ static FType ValidateOperatorAndGetResultType(FEmitter& Emitter, EOperator Op, F
 		case OR_MatrixMultiplyResult:
 		{
 			FType SecondArgumentType = Arguments[1]->Type;
-			if (FirstArgumentType.GetPrimitive().NumColumns != SecondArgumentType.GetPrimitive().NumRows)
+
+			FPrimitive LhsPrimitiveType = FirstArgumentType.GetPrimitive();
+			FPrimitive RhsPrimitiveType = SecondArgumentType.GetPrimitive();
+
+			int32 RhsRows = RhsPrimitiveType.NumRows;
+			int32 RhsColumns = RhsPrimitiveType.NumColumns;
+			int32 OutputRows = LhsPrimitiveType.NumRows;
+			int32 OutputColumns = RhsPrimitiveType.NumColumns;
+
+			// When multiplying matrix * vector, we reinterpret the input as a column vector (NumColumns == 1), even though by default our vectors are row vectors.
+			// And the output is reinterpreted back as a row vector.
+			if (SecondArgumentType.IsVector())
+			{
+				RhsRows = RhsPrimitiveType.NumColumns;
+				RhsColumns = 1;
+				OutputRows = 1;
+				OutputColumns = LhsPrimitiveType.NumRows;
+			}
+
+			if (FirstArgumentType.GetPrimitive().NumColumns != RhsRows)
 			{
 				Emitter.Errorf({}, TEXT("Cannot matrix multiply a '%s' with a '%s'."), *FirstArgumentType.GetSpelling(), *SecondArgumentType.GetSpelling());
 				return FType::MakePoison();
 			}
 
-			FPrimitive LhsPrimitiveType = FirstArgumentType.GetPrimitive();
-			FPrimitive RhsPrimitiveType = SecondArgumentType.GetPrimitive();
-
-			return FType::MakePrimitive(LhsPrimitiveType.ScalarKind, LhsPrimitiveType.NumRows, RhsPrimitiveType.NumColumns);
+			return FType::MakePrimitive(LhsPrimitiveType.ScalarKind, OutputRows, OutputColumns);
 		}
 	}
 
diff --git a/Engine/Source/Runtime/Engine/Private/Materials/MaterialIRModule.cpp b/Engine/Source/Runtime/Engine/Private/Materials/MaterialIRModule.cpp
--- a/Engine/Source/Runtime/Engine/Private/Materials/MaterialIRModule.cpp
+++ b/Engine/Source/Runtime/Engine/Private/Materials/MaterialIRModule.cpp
@@ -39,6 +39,8 @@ void FMaterialIRModule::Empty()
 
 	Statistics.NumVertexTexCoords = 0;
 	Statistics.NumPixelTexCoords = 0;
+
+	Statistics.DynamicParticleParameterMask = 0;
 }
 
 FStringView FMaterialIRModule::InternString(FStringView InString)
diff --git a/Engine/Source/Runtime/Engine/Private/Materials/MaterialIRModuleBuilder.cpp b/Engine/Source/Runtime/Engine/Private/Materials/MaterialIRModuleBuilder.cpp
--- a/Engine/Source/Runtime/Engine/Private/Materials/MaterialIRModuleBuilder.cpp
+++ b/Engine/Source/Runtime/Engine/Private/Materials/MaterialIRModuleBuilder.cpp
@@ -287,15 +287,22 @@ struct FMaterialIRModuleBuilderImpl
 	{
 		MIR::TTemporaryArray<MIR::FValue*> CallInputValues{ Call->FunctionInputs.Num() };
 
-		// Make sure each function input is connected and has a value. If so, cache the values flowing into this
-		// funcion call inside the auxiliary value array.
+		// Make sure each required function input is connected and has a value.  If so, cache the values flowing into this
+		// funcion call inside the auxiliary value array.  If the input is optional (bUsePreviewValueAsDefault set), we can
+		// ignore the missing value, and the downstream UMaterialExpressionFunctionInput build function will return the default.
 		for (int i = 0; i < Call->FunctionInputs.Num(); ++i)
 		{
 			FFunctionExpressionInput& FunctionInput = Call->FunctionInputs[i];
-			MIR::FValue* Value = Emitter->Input(Call->GetInput(i));
-			if (Value)
+			MIR::FType Type = MIR::FType::FromMaterialValueType(FunctionInput.ExpressionInput->GetInputValueType(0));
+
+			if (FunctionInput.ExpressionInput->bUsePreviewValueAsDefault)
+			{
+				MIR::FValue* Value = Emitter->TryInput(Call->GetInput(i));
+				CallInputValues[i] = Value ? Emitter->Cast(Value, Type) : nullptr;
+			}
+			else
 			{
-				MIR::FType Type = MIR::FType::FromMaterialValueType(FunctionInput.ExpressionInput->GetInputValueType(0));
+				MIR::FValue* Value = Emitter->Input(Call->GetInput(i));
 				CallInputValues[i] = Emitter->Cast(Value, Type);
 			}
 		}
@@ -318,14 +325,17 @@ struct FMaterialIRModuleBuilderImpl
 		// Forward values flowing into call inputs to called function inputs
 		for (int i = 0; i < Call->FunctionInputs.Num(); ++i)
 		{
-			FFunctionExpressionInput& FunctionInput = Call->FunctionInputs[i];
+			if (CallInputValues[i])
+			{
+				FFunctionExpressionInput& FunctionInput = Call->FunctionInputs[i];
 
-			// Bind the value flowing into the function call input to the function input
-			// expression (inside the function) in the new context.
-			NewContext.SetOutputValue(FunctionInput.ExpressionInput->GetOutput(0), CallInputValues[i]);
+				// Bind the value flowing into the function call input to the function input
+				// expression (inside the function) in the new context.
+				NewContext.SetOutputValue(FunctionInput.ExpressionInput->GetOutput(0), CallInputValues[i]);
 
-			// Mark the function input as built.
-			NewContext.BuiltExpressions.Add(FunctionInput.ExpressionInput.Get());
+				// Mark the function input as built.
+				NewContext.BuiltExpressions.Add(FunctionInput.ExpressionInput.Get());
+			}
 		}
 
 		// Finally push the function outputs to the expression evaluation stack in the new context.
diff --git a/Engine/Source/Runtime/Engine/Private/Materials/MaterialIRToHLSLTranslator.cpp b/Engine/Source/Runtime/Engine/Private/Materials/MaterialIRToHLSLTranslator.cpp
--- a/Engine/Source/Runtime/Engine/Private/Materials/MaterialIRToHLSLTranslator.cpp
+++ b/Engine/Source/Runtime/Engine/Private/Materials/MaterialIRToHLSLTranslator.cpp
@@ -541,6 +541,10 @@ struct FPrivate : FMaterialIRToHLSLTranslation
 			const bool bPrevious = (CurrentStage == MIR::EStage::Stage_Vertex && bCompilingPreviousFrame);
 			Printer << (bPrevious ? TEXTVIEW("GetPreviousActorTranslatedWorldPosition(Parameters)") : TEXTVIEW("GetActorTranslatedWorldPosition(Parameters)"));
 		}
+		else if (ExternalInput->Id == MIR::EExternalInput::DynamicParticleParameterIndex)
+		{
+			Printer << ExternalInput->UserData;
+		}
 		else
 		{
 			FStringView Code;
@@ -1369,6 +1373,7 @@ struct FPrivate : FMaterialIRToHLSLTranslation
 					case MIR::UO_Logarithm: OpString = TEXTVIEW("log"); break;
 					case MIR::UO_Logarithm10: OpString = TEXTVIEW("log10"); break;
 					case MIR::UO_Logarithm2: OpString = TEXTVIEW("log2"); break;
+					case MIR::UO_Negate: OpString = TEXTVIEW("-"); break;
 					case MIR::UO_Reciprocal: OpString = TEXTVIEW("rcp"); break;
 					case MIR::UO_Round: OpString = TEXTVIEW("round"); break;
 					case MIR::UO_Rsqrt: OpString = TEXTVIEW("rsqrt"); break;
@@ -2173,6 +2178,13 @@ struct FPrivate : FMaterialIRToHLSLTranslation
 		OutEnvironment.SetDefine(TEXT("IS_MATERIAL_SHADER"), true);
 		OutEnvironment.SetDefine(TEXT("VIRTUAL_TEXTURE_OUTPUT"), CompilationOutput.bHasRuntimeVirtualTextureOutputNode != 0);
 
+		uint32 DynamicParametersMask = Module->GetStatistics().DynamicParticleParameterMask;
+		if (DynamicParametersMask)
+		{
+			OutEnvironment.SetDefine(TEXT("USE_DYNAMIC_PARAMETERS"), 1);
+			OutEnvironment.SetDefine(TEXT("DYNAMIC_PARAMETERS_MASK"), DynamicParametersMask);
+		}
+
 		// Set all defines that are defined by the module.
 		// Any conditional exemption via material properties, such as 'Material->IsUsedWithInstancedStaticMeshes()', are handled during the material IR analysis.
 		for (const FName& EnvironmentDefine : Module->GetEnvironmentDefines())
diff --git a/Engine/Source/Runtime/Engine/Private/Materials/MaterialIRValueAnalyzer.cpp b/Engine/Source/Runtime/Engine/Private/Materials/MaterialIRValueAnalyzer.cpp
--- a/Engine/Source/Runtime/Engine/Private/Materials/MaterialIRValueAnalyzer.cpp
+++ b/Engine/Source/Runtime/Engine/Private/Materials/MaterialIRValueAnalyzer.cpp
@@ -63,6 +63,11 @@ static void AnalyzeExternalInput(FMaterialIRValueAnalyzer& Analyzer, MIR::FExter
 
 		Analyzer.CompilationOutput->bUsesGlobalDistanceField = true;
 	}
+
+	if (ExternalInput->Id == MIR::EExternalInput::DynamicParticleParameterIndex)
+	{
+		Analyzer.Module->GetStatistics().DynamicParticleParameterMask |= 1 << ExternalInput->UserData;
+	}
 }
 
 static EMaterialShaderFrequency MapToMaterialShaderFrequencyOrAny(MIR::EStage Stage)
diff --git a/Engine/Source/Runtime/Engine/Public/Materials/MaterialExpressionBounds.h b/Engine/Source/Runtime/Engine/Public/Materials/MaterialExpressionBounds.h
--- a/Engine/Source/Runtime/Engine/Public/Materials/MaterialExpressionBounds.h
+++ b/Engine/Source/Runtime/Engine/Public/Materials/MaterialExpressionBounds.h
@@ -5,6 +5,7 @@
 #include "CoreMinimal.h"
 #include "UObject/ObjectMacros.h"
 #include "Materials/MaterialExpression.h"
+#include "Materials/MaterialExpressionExternalCodeBase.h"
 #include "MaterialExpressionBounds.generated.h"
 
 UENUM()
@@ -21,7 +22,7 @@ enum EMaterialExpressionBoundsType : uint8
 };
 
 UCLASS(collapsecategories, hidecategories = Object, MinimalAPI)
-class UMaterialExpressionBounds : public UMaterialExpression
+class UMaterialExpressionBounds : public UMaterialExpressionExternalCodeBase
 {
 	GENERATED_UCLASS_BODY()
 
@@ -39,6 +40,7 @@ class UMaterialExpressionBounds : public UMaterialExpression
 
 #if WITH_EDITOR
 	//~ Begin UMaterialExpression Interface
+	virtual void Build(MIR::FEmitter& Emitter) override;
 	virtual int32 Compile(class FMaterialCompiler* Compiler, int32 OutputIndex) override;
 	virtual void GetCaption(TArray<FString>& OutCaptions) const override;
 	virtual void GetConnectorToolTip(int32 InputIndex, int32 OutputIndex, TArray<FString>& OutToolTip) override;
diff --git a/Engine/Source/Runtime/Engine/Public/Materials/MaterialExpressionBumpOffset.h b/Engine/Source/Runtime/Engine/Public/Materials/MaterialExpressionBumpOffset.h
--- a/Engine/Source/Runtime/Engine/Public/Materials/MaterialExpressionBumpOffset.h
+++ b/Engine/Source/Runtime/Engine/Public/Materials/MaterialExpressionBumpOffset.h
@@ -36,6 +36,7 @@ class UMaterialExpressionBumpOffset : public UMaterialExpression
 
 	//~ Begin UMaterialExpression Interface
 #if WITH_EDITOR
+	virtual void Build(MIR::FEmitter& Emitter) override;
 	virtual int32 Compile(class FMaterialCompiler* Compiler, int32 OutputIndex) override;
 	virtual void GetCaption(TArray<FString>& OutCaptions) const override;
 #endif
diff --git a/Engine/Source/Runtime/Engine/Public/Materials/MaterialExpressionDynamicParameter.h b/Engine/Source/Runtime/Engine/Public/Materials/MaterialExpressionDynamicParameter.h
--- a/Engine/Source/Runtime/Engine/Public/Materials/MaterialExpressionDynamicParameter.h
+++ b/Engine/Source/Runtime/Engine/Public/Materials/MaterialExpressionDynamicParameter.h
@@ -46,6 +46,7 @@ class UMaterialExpressionDynamicParameter : public UMaterialExpression
 
 	//~ Begin UMaterialExpression Interface
 #if WITH_EDITOR
+	virtual void Build(MIR::FEmitter& Emitter) override;
 	virtual int32 Compile(class FMaterialCompiler* Compiler, int32 OutputIndex) override;
 	virtual void GetCaption(TArray<FString>& OutCaptions) const override;
 	virtual TArray<FExpressionOutput>& GetOutputs() override;
diff --git a/Engine/Source/Runtime/Engine/Public/Materials/MaterialExpressionNoise.h b/Engine/Source/Runtime/Engine/Public/Materials/MaterialExpressionNoise.h
--- a/Engine/Source/Runtime/Engine/Public/Materials/MaterialExpressionNoise.h
+++ b/Engine/Source/Runtime/Engine/Public/Materials/MaterialExpressionNoise.h
@@ -7,6 +7,7 @@
 #include "UObject/ObjectMacros.h"
 #include "MaterialExpressionIO.h"
 #include "Materials/MaterialExpression.h"
+#include "Materials/MaterialExpressionExternalCodeBase.h"
 #include "MaterialExpressionNoise.generated.h"
 
 UENUM()
@@ -112,6 +113,7 @@ class UMaterialExpressionNoise : public UMaterialExpression
 
 	//~ Begin UMaterialExpression Interface
 #if WITH_EDITOR
+	virtual void Build(MIR::FEmitter& Emitter) override;
 	virtual bool CanEditChange(const FProperty* InProperty) const override;
 	virtual FName GetInputName(int32 InputIndex) const override;
 	virtual void PostEditChangeProperty(FPropertyChangedEvent& PropertyChangedEvent) override;
@@ -122,7 +124,7 @@ class UMaterialExpressionNoise : public UMaterialExpression
 };
 
 UCLASS(MinimalAPI)
-class UMaterialExpressionScalarBlueNoise : public UMaterialExpression
+class UMaterialExpressionScalarBlueNoise : public UMaterialExpressionExternalCodeBase
 {
 	GENERATED_UCLASS_BODY()
 
diff --git a/Engine/Source/Runtime/Engine/Public/Materials/MaterialExpressionObjectLocalBounds.h b/Engine/Source/Runtime/Engine/Public/Materials/MaterialExpressionObjectLocalBounds.h
--- a/Engine/Source/Runtime/Engine/Public/Materials/MaterialExpressionObjectLocalBounds.h
+++ b/Engine/Source/Runtime/Engine/Public/Materials/MaterialExpressionObjectLocalBounds.h
@@ -5,10 +5,11 @@
 #include "CoreMinimal.h"
 #include "UObject/ObjectMacros.h"
 #include "Materials/MaterialExpression.h"
+#include "Materials/MaterialExpressionExternalCodeBase.h"
 #include "MaterialExpressionObjectLocalBounds.generated.h"
 
 UCLASS(collapsecategories, hidecategories = Object)
-class UMaterialExpressionObjectLocalBounds : public UMaterialExpression
+class UMaterialExpressionObjectLocalBounds : public UMaterialExpressionExternalCodeBase
 {
 	GENERATED_UCLASS_BODY()
 
diff --git a/Engine/Source/Runtime/Engine/Public/Materials/MaterialExpressionObjectOrientation.h b/Engine/Source/Runtime/Engine/Public/Materials/MaterialExpressionObjectOrientation.h
--- a/Engine/Source/Runtime/Engine/Public/Materials/MaterialExpressionObjectOrientation.h
+++ b/Engine/Source/Runtime/Engine/Public/Materials/MaterialExpressionObjectOrientation.h
@@ -6,10 +6,11 @@
 #include "CoreMinimal.h"
 #include "UObject/ObjectMacros.h"
 #include "Materials/MaterialExpression.h"
+#include "Materials/MaterialExpressionExternalCodeBase.h"
 #include "MaterialExpressionObjectOrientation.generated.h"
 
 UCLASS(collapsecategories, hidecategories=Object, MinimalAPI)
-class UMaterialExpressionObjectOrientation : public UMaterialExpression
+class UMaterialExpressionObjectOrientation : public UMaterialExpressionExternalCodeBase
 {
 	GENERATED_BODY()
 
diff --git a/Engine/Source/Runtime/Engine/Public/Materials/MaterialExpressionPanner.h b/Engine/Source/Runtime/Engine/Public/Materials/MaterialExpressionPanner.h
--- a/Engine/Source/Runtime/Engine/Public/Materials/MaterialExpressionPanner.h
+++ b/Engine/Source/Runtime/Engine/Public/Materials/MaterialExpressionPanner.h
@@ -40,6 +40,7 @@ class UMaterialExpressionPanner : public UMaterialExpression
 
 	//~ Begin UMaterialExpression Interface
 #if WITH_EDITOR
+	virtual void Build(MIR::FEmitter& Emitter) override;
 	virtual int32 Compile(class FMaterialCompiler* Compiler, int32 OutputIndex) override;
 	virtual void GetCaption(TArray<FString>& OutCaptions) const override;
 	virtual bool NeedsRealtimePreview() override { return Time.Expression==NULL && (SpeedX != 0.f || SpeedY != 0.f); }
diff --git a/Engine/Source/Runtime/Engine/Public/Materials/MaterialExpressionParticlePositionWS.h b/Engine/Source/Runtime/Engine/Public/Materials/MaterialExpressionParticlePositionWS.h
--- a/Engine/Source/Runtime/Engine/Public/Materials/MaterialExpressionParticlePositionWS.h
+++ b/Engine/Source/Runtime/Engine/Public/Materials/MaterialExpressionParticlePositionWS.h
@@ -6,10 +6,11 @@
 #include "CoreMinimal.h"
 #include "UObject/ObjectMacros.h"
 #include "Materials/MaterialExpression.h"
+#include "Materials/MaterialExpressionExternalCodeBase.h"
 #include "MaterialExpressionParticlePositionWS.generated.h"
 
 UCLASS(collapsecategories, hidecategories=Object)
-class UMaterialExpressionParticlePositionWS : public UMaterialExpression
+class UMaterialExpressionParticlePositionWS : public UMaterialExpressionExternalCodeBase
 {
 	GENERATED_BODY()
 
@@ -20,6 +21,7 @@ public:
 
 	//~ Begin UMaterialExpression Interface
 #if WITH_EDITOR
+	virtual void Build(MIR::FEmitter& Emitter) override;
 	virtual int32 Compile(class FMaterialCompiler* Compiler, int32 OutputIndex) override;
 	virtual void GetCaption(TArray<FString>& OutCaptions) const override;
 #endif
diff --git a/Engine/Source/Runtime/Engine/Public/Materials/MaterialExpressionPerInstanceCustomData.h b/Engine/Source/Runtime/Engine/Public/Materials/MaterialExpressionPerInstanceCustomData.h
--- a/Engine/Source/Runtime/Engine/Public/Materials/MaterialExpressionPerInstanceCustomData.h
+++ b/Engine/Source/Runtime/Engine/Public/Materials/MaterialExpressionPerInstanceCustomData.h
@@ -5,10 +5,11 @@
 #include "CoreMinimal.h"
 #include "UObject/ObjectMacros.h"
 #include "Materials/MaterialExpression.h"
+#include "Materials/MaterialExpressionExternalCodeBase.h"
 #include "MaterialExpressionPerInstanceCustomData.generated.h"
 
 UCLASS(collapsecategories, hidecategories=Object)
-class UMaterialExpressionPerInstanceCustomData : public UMaterialExpression
+class UMaterialExpressionPerInstanceCustomData : public UMaterialExpressionExternalCodeBase
 {
 	GENERATED_BODY()
 
@@ -25,6 +26,7 @@ public:
 
 	//~ Begin UMaterialExpression Interface
 #if WITH_EDITOR
+	virtual void Build(MIR::FEmitter& Emitter) override;
 	virtual int32 Compile(class FMaterialCompiler* Compiler, int32 OutputIndex) override;
 	virtual void GetCaption(TArray<FString>& OutCaptions) const override;
 
@@ -33,7 +35,7 @@ public:
 };
 
 UCLASS(collapsecategories, hidecategories=Object)
-class UMaterialExpressionPerInstanceCustomData3Vector : public UMaterialExpression
+class UMaterialExpressionPerInstanceCustomData3Vector : public UMaterialExpressionExternalCodeBase
 {
 	GENERATED_BODY()
 
@@ -50,6 +52,7 @@ public:
 
 	//~ Begin UMaterialExpression Interface
 #if WITH_EDITOR
+	virtual void Build(MIR::FEmitter& Emitter) override;
 	virtual int32 Compile(class FMaterialCompiler* Compiler, int32 OutputIndex) override;
 	virtual void GetCaption(TArray<FString>& OutCaptions) const override;
 
diff --git a/Engine/Source/Runtime/Engine/Public/Materials/MaterialExpressionPrecomputedAOMask.h b/Engine/Source/Runtime/Engine/Public/Materials/MaterialExpressionPrecomputedAOMask.h
--- a/Engine/Source/Runtime/Engine/Public/Materials/MaterialExpressionPrecomputedAOMask.h
+++ b/Engine/Source/Runtime/Engine/Public/Materials/MaterialExpressionPrecomputedAOMask.h
@@ -10,10 +10,11 @@
 #include "CoreMinimal.h"
 #include "UObject/ObjectMacros.h"
 #include "Materials/MaterialExpression.h"
+#include "Materials/MaterialExpressionExternalCodeBase.h"
 #include "MaterialExpressionPrecomputedAOMask.generated.h"
 
 UCLASS(collapsecategories, hidecategories=Object)
-class UMaterialExpressionPrecomputedAOMask : public UMaterialExpression
+class UMaterialExpressionPrecomputedAOMask : public UMaterialExpressionExternalCodeBase
 {
 	GENERATED_UCLASS_BODY()
 
diff --git a/Engine/Source/Runtime/Engine/Public/Materials/MaterialExpressionVectorNoise.h b/Engine/Source/Runtime/Engine/Public/Materials/MaterialExpressionVectorNoise.h
--- a/Engine/Source/Runtime/Engine/Public/Materials/MaterialExpressionVectorNoise.h
+++ b/Engine/Source/Runtime/Engine/Public/Materials/MaterialExpressionVectorNoise.h
@@ -81,6 +81,7 @@ class UMaterialExpressionVectorNoise : public UMaterialExpression
 
 	//~ Begin UMaterialExpression Interface
 #if WITH_EDITOR
+	virtual void Build(MIR::FEmitter& Emitter) override;
 	virtual bool CanEditChange(const FProperty* InProperty) const override;
 	virtual FName GetInputName(int32 InputIndex) const override;
 	virtual void PostEditChangeProperty(FPropertyChangedEvent& PropertyChangedEvent) override;
diff --git a/Engine/Source/Runtime/Engine/Public/Materials/MaterialIR.h b/Engine/Source/Runtime/Engine/Public/Materials/MaterialIR.h
--- a/Engine/Source/Runtime/Engine/Public/Materials/MaterialIR.h
+++ b/Engine/Source/Runtime/Engine/Public/Materials/MaterialIR.h
@@ -417,6 +417,7 @@ enum class EExternalInput
 	ViewMaterialTextureMipBias,
 	ViewMaterialTextureDerivativeMultiply,
 	GlobalDistanceField,					// Flag to mark that global distance fields are used, for value analyzer step.  Doesn't generate code.
+	DynamicParticleParameterIndex,			// UserData stores the parameter index.
 
 	Count,
 };
@@ -467,6 +468,7 @@ bool IsExternalInputWorldPosition(EExternalInput Id);
 struct FExternalInput : TValue<VK_ExternalInput>
 {
 	EExternalInput Id;
+	uint32 UserData;				// Optional user data for certain external input types
 };
 
 // It represents a Material Parameter Collection object value.
diff --git a/Engine/Source/Runtime/Engine/Public/Materials/MaterialIREmitter.h b/Engine/Source/Runtime/Engine/Public/Materials/MaterialIREmitter.h
--- a/Engine/Source/Runtime/Engine/Public/Materials/MaterialIREmitter.h
+++ b/Engine/Source/Runtime/Engine/Public/Materials/MaterialIREmitter.h
@@ -389,7 +389,7 @@ public:
 	FValueRef Poison();
 
 	// Emits an external input value of given input Id and returns it.
-	FValueRef ExternalInput(EExternalInput Id);
+	FValueRef ExternalInput(EExternalInput Id, uint32 UserData = 0);
 
 	// Emits a material parameter collection.
 	FValueRef MaterialParameterCollection(UMaterialParameterCollection* Collection);
diff --git a/Engine/Source/Runtime/Engine/Public/Materials/MaterialIRModule.h b/Engine/Source/Runtime/Engine/Public/Materials/MaterialIRModule.h
--- a/Engine/Source/Runtime/Engine/Public/Materials/MaterialIRModule.h
+++ b/Engine/Source/Runtime/Engine/Public/Materials/MaterialIRModule.h
@@ -41,6 +41,9 @@ public:
 
 		// Number of pixel texture coordinates used.
 		int32 NumPixelTexCoords;
+
+		// Dynamic particle parameter indices used by this material.
+		uint32 DynamicParticleParameterMask;
 	};
 
 public:
