Disable Python script plugin for GatherText commandlet.

diff --git a/Engine/Config/BaseEditor.ini b/Engine/Config/BaseEditor.ini
--- a/Engine/Config/BaseEditor.ini
+++ b/Engine/Config/BaseEditor.ini
@@ -28,6 +28,7 @@ SdfFontAtlasSize=2048
 [PythonScriptPlugin]
 +DisablePythonForCommandlet=CookShaders
 +DisablePythonForCommandlet=Cook
++DisablePythonForCommandlet=GatherText
 
 [AssetTools]
 ; Full Package Path being hidden from the Asset view
