Unshelved from pending changelist '45015913':

Updated device profiles for standalone XR devices. They no longer inherit from each other in sequence, and the new Android_OpenXR device profile matches the Quest 3 profile which has had its shadow quality increased for dynamic shadows to draw by default in Mobile Preview and on Device. Consider lowering this value if dynamic shadows are not required.

diff --git a/Engine/Config/BaseDeviceProfiles.ini b/Engine/Config/BaseDeviceProfiles.ini
--- a/Engine/Config/BaseDeviceProfiles.ini
+++ b/Engine/Config/BaseDeviceProfiles.ini
@@ -1280,35 +1280,40 @@ BaseProfileName=Android_High
 [Android_OpenXR DeviceProfile]
 DeviceType=Android
 BaseProfileName=Android_Mid
++CVars=r.Android.DisableVulkanSupport=0
 +CVars=vr.MobileMultiView=1
 +CVars=r.Mobile.AntiAliasing=3
 +CVars=r.ForceStartInVR=1
++CVars=xr.VRS.DynamicFoveation=1
++CVars=r.Vulkan.VRSFormat=3
++CVars=sg.ShadowQuality=2
 
 [Meta_Quest_3 DeviceProfile]
 DeviceType=Android
-BaseProfileName=Meta_Quest_Pro
+BaseProfileName=Android_OpenXR
++CVars=r.Mobile.Oculus.ForceSymmetric=1
++CVars=r.Mobile.AdrenoOcclusionMode=1
 
 [Meta_Quest_Pro DeviceProfile]
 DeviceType=Android
-BaseProfileName=Oculus_Quest2
+BaseProfileName=Android_OpenXR
++CVars=r.Mobile.Oculus.ForceSymmetric=1
++CVars=r.Mobile.AdrenoOcclusionMode=1
 
 [Oculus_Quest2 DeviceProfile]
 DeviceType=Android
-BaseProfileName=Oculus_Quest
+BaseProfileName=Android_OpenXR
 +CVars=r.Mobile.Oculus.ForceSymmetric=1
-+CVars=fx.NiagaraAllowGPUParticles=1
-+CVars=FX.AllowGPUSorting=1
 +CVars=r.Mobile.AdrenoOcclusionMode=1
++CVars=sg.ShadowQuality=1
 
 [Oculus_Quest DeviceProfile]
 DeviceType=Android
 BaseProfileName=Android_OpenXR
-+CVars=xr.VRS.DynamicFoveation=1
-+CVars=r.Android.DisableVulkanSupport=0
 ; disable running Niagara compute shaders on the GPU. This can cause hard to trace crashes.
 +CVars=fx.NiagaraAllowGPUParticles=0
 +CVars=FX.AllowGPUSorting=0
-+CVars=r.Vulkan.VRSFormat=3
++CVars=sg.ShadowQuality=1
 
 [Mac DeviceProfile]
 DeviceType=Mac
