Default PC GDK to simplified user model

#rnx

diff --git a/Engine/Config/BaseEngine.ini b/Engine/Config/BaseEngine.ini
--- a/Engine/Config/BaseEngine.ini
+++ b/Engine/Config/BaseEngine.ini
@@ -3944,3 +3944,7 @@ VerificationThreadCount=4
 
 [/Script/ControlRig.ControlRigSettings]
 DefaultRootModule=/ControlRig/Modules/Modules56/Root.Root
+
+[/Script/MSGamingSupport.MSGamingSettings]
+bUseSimplifiedUserModel=True
+
