Lumen - enable half res integration on High scalability (r.Lumen.ScreenProbeGather.StochasticInterpolation 2). This can soften normals in indirect lighting, and make GI a bit more noisy. On the other hand it saves ~0.5ms on console at 1080p, which feels like a right tradeoff for High scalability intended for 60hz on console.

diff --git a/Engine/Config/BaseScalability.ini b/Engine/Config/BaseScalability.ini
--- a/Engine/Config/BaseScalability.ini
+++ b/Engine/Config/BaseScalability.ini
@@ -297,6 +297,7 @@ r.Lumen.ScreenProbeGather.NumAdaptiveProbes=16
 r.Lumen.ScreenProbeGather.TracingOctahedronResolution=8
 r.Lumen.ScreenProbeGather.IrradianceFormat=1
 r.Lumen.ScreenProbeGather.StochasticInterpolation=1
+r.Lumen.ScreenProbeGather.IntegrateDownsampleFactor=2
 r.Lumen.ScreenProbeGather.FullResolutionJitterWidth=1
 r.Lumen.ScreenProbeGather.TwoSidedFoliageBackfaceDiffuse=0
 r.Lumen.ScreenProbeGather.ScreenTraces.HZBTraversal.FullResDepth=1
@@ -328,6 +329,7 @@ r.Lumen.ScreenProbeGather.NumAdaptiveProbes=8
 r.Lumen.ScreenProbeGather.TracingOctahedronResolution=8
 r.Lumen.ScreenProbeGather.IrradianceFormat=0
 r.Lumen.ScreenProbeGather.StochasticInterpolation=0
+r.Lumen.ScreenProbeGather.IntegrateDownsampleFactor=1
 r.Lumen.ScreenProbeGather.FullResolutionJitterWidth=1
 r.Lumen.ScreenProbeGather.TwoSidedFoliageBackfaceDiffuse=1
 r.Lumen.ScreenProbeGather.ScreenTraces.HZBTraversal.FullResDepth=1
@@ -360,6 +362,7 @@ r.Lumen.ScreenProbeGather.NumAdaptiveProbes=4
 r.Lumen.ScreenProbeGather.TracingOctahedronResolution=16
 r.Lumen.ScreenProbeGather.IrradianceFormat=0
 r.Lumen.ScreenProbeGather.StochasticInterpolation=0
+r.Lumen.ScreenProbeGather.IntegrateDownsampleFactor=1
 r.Lumen.ScreenProbeGather.FullResolutionJitterWidth=.5
 r.Lumen.ScreenProbeGather.TwoSidedFoliageBackfaceDiffuse=1
 r.Lumen.ScreenProbeGather.ScreenTraces.HZBTraversal.FullResDepth=1
