Change the default of bUseLinkedOrthographicViewports to `False`.

[RN] Minor, Editor

diff --git a/Engine/Config/BaseEditorPerProjectUserSettings.ini b/Engine/Config/BaseEditorPerProjectUserSettings.ini
--- a/Engine/Config/BaseEditorPerProjectUserSettings.ini
+++ b/Engine/Config/BaseEditorPerProjectUserSettings.ini
@@ -454,7 +454,7 @@ bEnableViewportHoverFeedback=False
 ; If enabled, selected objects will be highlighted with brackets in all modes rather than a special highlight color.
 bHighlightWithBrackets=False
 ; If true all orthographic viewports are linked to the same position and move together
-bUseLinkedOrthographicViewports=True
+bUseLinkedOrthographicViewports=False
 ; If true, objects must be entirely encompassed by the selection box in ortho. viewports to be selected
 bStrictBoxSelection=False
 ; True if viewport box selection also selects occluded objects, false if only objects with visible pixels are selected
