Add new iPhone17 models. Verified that all iPhone17 models are still 460ppi

diff --git a/Engine/Config/BaseDeviceProfiles.ini b/Engine/Config/BaseDeviceProfiles.ini
--- a/Engine/Config/BaseDeviceProfiles.ini
+++ b/Engine/Config/BaseDeviceProfiles.ini
@@ -61,6 +61,10 @@
 +DeviceProfileNameAndTypes=iPhone16Plus,IOS
 +DeviceProfileNameAndTypes=iPhone16Pro,IOS
 +DeviceProfileNameAndTypes=iPhone16ProMax,IOS
++DeviceProfileNameAndTypes=iPhone17,IOS
++DeviceProfileNameAndTypes=iPhone17Pro,IOS
++DeviceProfileNameAndTypes=iPhone17ProMax,IOS
++DeviceProfileNameAndTypes=iPhoneAir,IOS
 +DeviceProfileNameAndTypes=iPadPro,IOS
 +DeviceProfileNameAndTypes=iPadPro105,IOS
 +DeviceProfileNameAndTypes=iPadPro129,IOS
@@ -268,6 +272,10 @@ iPhone17,2=iPhone16ProMax
 iPhone17,3=iPhone16
 iPhone17,4=iPhone16Plus
 iPhone17,5=iPhone16e
+iPhone18,1=iPhone17Pro
+iPhone18,2=iPhone17ProMax
+iPhone18,3=iPhone17
+iPhone18,4=iPhoneAir
 iPod9,1=iPodTouch7
 iPad5,[1,2]=iPadMini4
 iPad5,[3,4]=iPadAir2
@@ -563,6 +571,22 @@ BaseProfileName=iPhone15Pro
 DeviceType=IOS
 BaseProfileName=iPhone15ProMax
 
+[iPhone17Pro DeviceProfile]
+DeviceType=IOS
+BaseProfileName=iPhone15Pro
+
+[iPhone17ProMax DeviceProfile]
+DeviceType=IOS
+BaseProfileName=iPhone15ProMax
+
+[iPhone17 DeviceProfile]
+DeviceType=IOS
+BaseProfileName=iPhone15
+
+[iPhoneAir DeviceProfile]
+DeviceType=IOS
+BaseProfileName=iPhone15Pro
+
 [iPodTouch7 DeviceProfile]
 DeviceType=IOS
 BaseProfileName=iPhone7
