Updated settings for a default mobile device profiles
Added Epic bucket for a modern devices with a most graphics features on

diff --git a/Engine/Config/Android/AndroidEngine.ini b/Engine/Config/Android/AndroidEngine.ini
--- a/Engine/Config/Android/AndroidEngine.ini
+++ b/Engine/Config/Android/AndroidEngine.ini
@@ -66,13 +66,10 @@ DefaultPlatformService=AndroidLocalNotification
 [ConsoleVariables]
 Slate.CacheRenderData=0
 r.VSync=1
-r.ParallelShadows=0
 ; high priority thread tasks can be downgraded to normal priority if a task is created with EPowerSavingEligibility::Eligible.
 TaskGraph.EnablePowerSavingThreadPriorityReduction=1
 ; prefer D24_S8 for depth/stencil
 r.Vulkan.Depth24Bit=1
-;GalaxyS6Edge=577
-r.Shadow.ForceSerialSingleRenderPass=1
 ; Support Mali 64k texel buffer limitation, tuned for 1080p. If you target 2k-4k resolutions please also set the LightGridPixelSize to 128
 r.Forward.LightGridSizeZ = 8
 r.Forward.MaxCulledLightsPerCell = 8
@@ -105,13 +102,6 @@ r.Vulkan.AllowSparseResidency=0
 s.StreamableStripDebugName=true
 s.StreamableStripDebugNameInShipping=true
 
-[/Script/Engine.RendererSettings]
-r.DistanceFields=1
-r.TemporalAA.Mobile.UseCompute=0
-; Disables parallel GDME by default due to not enough threads
-r.Visibility.DynamicMeshElements.Parallel=0
-r.Shadow.DetectVertexShaderLayerAtRuntime=1
-
 [AlternateTextureCompression]
 ;disabling TextureCompressionFormat for Android/IOS no longer needed
 ;TextureCompressionFormat=""
diff --git a/Engine/Config/Android/AndroidScalability.ini b/Engine/Config/Android/AndroidScalability.ini
--- a/Engine/Config/Android/AndroidScalability.ini
+++ b/Engine/Config/Android/AndroidScalability.ini
@@ -53,6 +53,7 @@ r.Shadow.CSM.MaxCascades=1
 r.Shadow.RadiusThreshold=0.06
 r.VolumetricFog=0
 r.ContactShadows=0
+r.DistanceFieldShadowing=0
 r.DFShadowQuality=0
 
 [ShadowQuality@1]
@@ -63,6 +64,7 @@ r.Shadow.CSM.MaxCascades=1
 r.Shadow.RadiusThreshold=0.06
 r.VolumetricFog=0
 r.ContactShadows=0
+r.DistanceFieldShadowing=0
 r.DFShadowQuality=0
 
 [ShadowQuality@2]
@@ -73,7 +75,7 @@ r.Shadow.CSM.MaxCascades=2
 r.Shadow.RadiusThreshold=0.06
 r.VolumetricFog=0
 r.ContactShadows=0
-r.DFShadowQuality=0
+r.DFShadowQuality=1
 r.Mobile.SupportInsetShadows=True
 
 [ShadowQuality@3]
@@ -145,15 +147,19 @@ r.Mobile.AmbientOcclusionQuality=3
 r.LensFlareQuality=2
 
 [AntiAliasingQuality@0]
+r.Mobile.AntiAliasing=0
 r.FXAA.Quality=0
 
 [AntiAliasingQuality@1]
+r.Mobile.AntiAliasing=1
 r.FXAA.Quality=0
 
 [AntiAliasingQuality@2]
+r.Mobile.AntiAliasing=2
 r.FXAA.Quality=0
 
 [AntiAliasingQuality@3]
+r.Mobile.AntiAliasing=2
 r.FXAA.Quality=0
 
 ; ------------------------------------------------------------------------------
diff --git a/Engine/Config/Android/BaseAndroidEngine.ini b/Engine/Config/Android/BaseAndroidEngine.ini
--- a/Engine/Config/Android/BaseAndroidEngine.ini
+++ b/Engine/Config/Android/BaseAndroidEngine.ini
@@ -23,8 +23,13 @@ gc.MaxObjectsInGame=131072
 NumPreallocatedVertices=200
 
 [/Script/Engine.RendererSettings]
+r.DistanceFields=1
+r.TemporalAA.Mobile.UseCompute=0
+; Disables parallel GDME by default due to not enough threads
+r.Visibility.DynamicMeshElements.Parallel=0
 ; Volumetric Clouds are not supported by mobile renderer, skip materials/global shaders compilation.
 r.VolumetricCloud.Support=0
+r.OIT.SortedTriangles=0
 
 [/Script/AndroidRuntimeSettings.AndroidRuntimeSettings]
 ; Default to latest astc version - projects can manually specify this to retain older versions to prevent
diff --git a/Engine/Config/BaseDeviceProfiles.ini b/Engine/Config/BaseDeviceProfiles.ini
--- a/Engine/Config/BaseDeviceProfiles.ini
+++ b/Engine/Config/BaseDeviceProfiles.ini
@@ -349,425 +349,456 @@
 +CVars=r.MorphTarget.Mode=0
 ; PF_B8G8R8A8
 +CVars=r.DefaultBackBufferPixelFormat=0 
-; Scalability
-+CVars=sg.AntiAliasingQuality=0
-+CVars=sg.ShadowQuality=1
-+CVars=sg.GlobalIlluminationQuality=1
-+CVars=sg.ReflectionQuality=1
-+CVars=sg.ViewDistanceQuality=3
-+CVars=sg.PostProcessQuality=3
-+CVars=sg.TextureQuality=3
-+CVars=sg.EffectsQuality=3
-+CVars=sg.FoliageQuality=3
-+CVars=sg.LandscapeQuality=2
 ; LightGrid Settings
 +CVars=r.Forward.LightLinkedListCulling=0
 ; Optimize Slate signed distance field text for performance
 +CVars=UI.SlateSDFText.RasterizationMode=SdfApproximation
 +CVars=UI.SlateSDFText.ResolutionLevel=1
 +CVars=AsyncTraceDelegateHitchThresholdMS=0.75
+; Use 'OcclusionFeedback' instead of occlusion queires to reduce drawcalls
++CVars=r.OcclusionFeedback.Enable=1
+; Max 256k partices
++CVars=fx.GPUSimulationTextureSizeX=512
++CVars=fx.GPUSimulationTextureSizeY=512
++CVars=r.Streaming.MaxTempMemoryAllowed=25
 
 ; PreviewAllowlistCVars and PreviewDenyListCVars are arrays of cvars (or prefixes) that are included or excluded from being applied in editor preview
 ; If any PreviewAllowlistCVars is set, cvars are denied by default.
 PreviewAllowlistCVars=sg.
 
+[IOS_Low DeviceProfile]
+BaseProfileName=IOS
++CVars=r.MobileContentScaleFactor=2
++CVars=r.NumBufferedOcclusionQueries=2
++CVars=sg.ViewDistanceQuality=0
++CVars=sg.AntiAliasingQuality=0
++CVars=sg.ShadowQuality=1
++CVars=sg.GlobalIlluminationQuality=0
++CVars=sg.ReflectionQuality=0
++CVars=sg.PostProcessQuality=1
++CVars=sg.TextureQuality=0
++CVars=sg.EffectsQuality=0
++CVars=sg.FoliageQuality=0
+
+[IOS_Mid DeviceProfile]
+BaseProfileName=IOS
++CVars=r.MobileContentScaleFactor=2
++CVars=sg.ViewDistanceQuality=1
++CVars=sg.AntiAliasingQuality=1
++CVars=sg.ShadowQuality=1
++CVars=sg.GlobalIlluminationQuality=1
++CVars=sg.ReflectionQuality=1
++CVars=sg.PostProcessQuality=1
++CVars=sg.TextureQuality=1
++CVars=sg.EffectsQuality=1
++CVars=sg.FoliageQuality=1
+
+[IOS_High DeviceProfile]
+BaseProfileName=IOS
++CVars=r.MobileContentScaleFactor=2
++CVars=sg.ViewDistanceQuality=2
++CVars=sg.AntiAliasingQuality=2
++CVars=sg.ShadowQuality=2
++CVars=sg.GlobalIlluminationQuality=2
++CVars=sg.ReflectionQuality=2
++CVars=sg.PostProcessQuality=2
++CVars=sg.TextureQuality=2
++CVars=sg.EffectsQuality=2
++CVars=sg.FoliageQuality=2
+
+[IOS_Epic DeviceProfile]
+BaseProfileName=IOS
++CVars=r.MobileContentScaleFactor=2
++CVars=sg.ViewDistanceQuality=3
++CVars=sg.AntiAliasingQuality=3
++CVars=sg.ShadowQuality=3
++CVars=sg.GlobalIlluminationQuality=3
++CVars=sg.ReflectionQuality=3
++CVars=sg.PostProcessQuality=3
++CVars=sg.TextureQuality=3
++CVars=sg.EffectsQuality=3
++CVars=sg.FoliageQuality=3
+
 [iPhone6S DeviceProfile]
 DeviceType=IOS
-BaseProfileName=IOS
+BaseProfileName=IOS_Low
 +CVars=ios.PhysicalScreenDensity=326
-+CVars=r.MobileContentScaleFactor=2
-+CVars=sg.PostProcessQuality=3
-+CVars=sg.ShadowQuality=2
-+CVars=sg.GlobalIlluminationQuality=2
 
 [iPhone6SPlus DeviceProfile]
 DeviceType=IOS
-BaseProfileName=IOS
+BaseProfileName=IOS_Low
 +CVars=ios.PhysicalScreenDensity=401
-+CVars=r.MobileContentScaleFactor=2
-+CVars=sg.PostProcessQuality=3
-+CVars=sg.ShadowQuality=2
-+CVars=sg.GlobalIlluminationQuality=2
 
 [iPhoneSE DeviceProfile]
 DeviceType=IOS
-BaseProfileName=IOS
+BaseProfileName=IOS_Low
 +CVars=ios.PhysicalScreenDensity=326
-+CVars=r.MobileContentScaleFactor=2
-+CVars=sg.PostProcessQuality=3
 
 [iPhone7 DeviceProfile]
 DeviceType=IOS
-BaseProfileName=iPhone6S
+BaseProfileName=IOS_Mid
 +CVars=ios.PhysicalScreenDensity=326
-+CVars=r.MobileContentScaleFactor=2
 
 [iPhone7Plus DeviceProfile]
 DeviceType=IOS
-BaseProfileName=iPhone6SPlus
+BaseProfileName=IOS_Mid
 +CVars=ios.PhysicalScreenDensity=401
-+CVars=r.MobileContentScaleFactor=2
 
 [iPhone8 DeviceProfile]
 DeviceType=IOS
-BaseProfileName=iPhone7
+BaseProfileName=IOS_High
 +CVars=ios.PhysicalScreenDensity=326
-+CVars=sg.ShadowQuality=3
-+CVars=sg.GlobalIlluminationQuality=3
 
 [iPhone8Plus DeviceProfile]
 DeviceType=IOS
-BaseProfileName=iPhone7Plus
+BaseProfileName=IOS_High
 +CVars=ios.PhysicalScreenDensity=401
-+CVars=sg.ShadowQuality=3
-+CVars=sg.GlobalIlluminationQuality=3
 
 [iPhoneX DeviceProfile]
 DeviceType=IOS
-BaseProfileName=iPhone8
+BaseProfileName=IOS_High
 +CVars=ios.PhysicalScreenDensity=401
 +CVars=r.CustomUnsafeZones="(L:free[0,-15][812,15]);(P:fixed[83,0][206,30])"
 
 [iPhoneXS DeviceProfile]
 DeviceType=IOS
-BaseProfileName=iPhone8
+BaseProfileName=IOS_High
 +CVars=ios.PhysicalScreenDensity=401
 +CVars=r.CustomUnsafeZones="(L:free[0,-15][812,15]);(P:fixed[83,0][206,30])"
 
 [iPhoneXSMax DeviceProfile]
 DeviceType=IOS
-BaseProfileName=iPhone8
+BaseProfileName=IOS_High
 +CVars=ios.PhysicalScreenDensity=401
 +CVars=r.CustomUnsafeZones="(L:free[0,-15][896,15]);(P:fixed[104,0][206,30])"
 
 [iPhoneXR DeviceProfile]
 DeviceType=IOS
-BaseProfileName=iPhone8
+BaseProfileName=IOS_High
 +CVars=ios.PhysicalScreenDensity=326
 +CVars=r.CustomUnsafeZones="(L:free[0,-15][896,15]);(P:fixed[104,0][206,30])"
 
 [iPhone11 DeviceProfile]
 DeviceType=IOS
-BaseProfileName=iPhone8
+BaseProfileName=IOS_High
 +CVars=ios.PhysicalScreenDensity=326
 +CVars=r.CustomUnsafeZones="(L:free[0,-15][812,15]);(P:fixed[83,0][206,30])"
 
 [iPhone11Pro DeviceProfile]
 DeviceType=IOS
-BaseProfileName=iPhone8
+BaseProfileName=IOS_High
 +CVars=ios.PhysicalScreenDensity=401
 +CVars=r.CustomUnsafeZones="(L:free[0,-15][812,15]);(P:fixed[83,0][206,30])"
 
 [iPhone11ProMax DeviceProfile]
 DeviceType=IOS
-BaseProfileName=iPhone8
+BaseProfileName=IOS_High
 +CVars=ios.PhysicalScreenDensity=401
 +CVars=r.CustomUnsafeZones="(L:free[0,-15][812,15]);(P:fixed[83,0][206,30])"
 
 [iPhoneSE2 DeviceProfile]
 DeviceType=IOS
-BaseProfileName=iPhone11
+BaseProfileName=IOS_High
 +CVars=ios.PhysicalScreenDensity=326
 
 [iPhone12Mini DeviceProfile]
 DeviceType=IOS
-BaseProfileName=iPhone11
+BaseProfileName=IOS_Epic
 +CVars=ios.PhysicalScreenDensity=476
 
 [iPhone12 DeviceProfile]
 DeviceType=IOS
-BaseProfileName=iPhone11
+BaseProfileName=IOS_Epic
 +CVars=ios.PhysicalScreenDensity=460
 
 [iPhone12Pro DeviceProfile]
 DeviceType=IOS
-BaseProfileName=iPhone11Pro
+BaseProfileName=IOS_Epic
 +CVars=ios.PhysicalScreenDensity=460
 
 [iPhone12ProMax DeviceProfile]
 DeviceType=IOS
-BaseProfileName=iPhone11ProMax
+BaseProfileName=IOS_Epic
 +CVars=ios.PhysicalScreenDensity=458
 
 [iPhone13Mini DeviceProfile]
 DeviceType=IOS
-BaseProfileName=iPhone11
+BaseProfileName=IOS_Epic
 +CVars=ios.PhysicalScreenDensity=476
 
 [iPhone13 DeviceProfile]
 DeviceType=IOS
-BaseProfileName=iPhone11
+BaseProfileName=IOS_Epic
 +CVars=ios.PhysicalScreenDensity=460
 
 [iPhone13Pro DeviceProfile]
 DeviceType=IOS
-BaseProfileName=iPhone11Pro
+BaseProfileName=IOS_Epic
 +CVars=ios.PhysicalScreenDensity=460
 
 [iPhone13ProMax DeviceProfile]
 DeviceType=IOS
-BaseProfileName=iPhone11ProMax
+BaseProfileName=IOS_Epic
 +CVars=ios.PhysicalScreenDensity=458
 
 [iPhoneSE3 DeviceProfile]
 DeviceType=IOS
-BaseProfileName=iPhone13
+BaseProfileName=IOS_Epic
 +CVars=ios.PhysicalScreenDensity=326
 
 [iPhone14 DeviceProfile]
 DeviceType=IOS
-BaseProfileName=iPhone13
+BaseProfileName=IOS_Epic
++CVars=ios.PhysicalScreenDensity=460
 
 [iPhone14Plus DeviceProfile]
 DeviceType=IOS
-BaseProfileName=iPhone13
+BaseProfileName=IOS_Epic
 +CVars=ios.PhysicalScreenDensity=458
 
 [iPhone14Pro DeviceProfile]
 DeviceType=IOS
-BaseProfileName=iPhone13Pro
+BaseProfileName=IOS_Epic
++CVars=ios.PhysicalScreenDensity=460
 
 [iPhone14ProMax DeviceProfile]
 DeviceType=IOS
-BaseProfileName=iPhone13ProMax
+BaseProfileName=IOS_Epic
 +CVars=ios.PhysicalScreenDensity=460
 
 [iPhone15 DeviceProfile]
 DeviceType=IOS
-BaseProfileName=iPhone13
+BaseProfileName=IOS_Epic
++CVars=ios.PhysicalScreenDensity=460
 
 [iPhone15Plus DeviceProfile]
 DeviceType=IOS
-BaseProfileName=iPhone13
+BaseProfileName=IOS_Epic
 +CVars=ios.PhysicalScreenDensity=458
 
 [iPhone15Pro DeviceProfile]
 DeviceType=IOS
-BaseProfileName=iPhone13Pro
+BaseProfileName=IOS_Epic
++CVars=ios.PhysicalScreenDensity=460
 
 [iPhone15ProMax DeviceProfile]
 DeviceType=IOS
-BaseProfileName=iPhone13ProMax
+BaseProfileName=IOS_Epic
 +CVars=ios.PhysicalScreenDensity=460
 
 [iPhone16 DeviceProfile]
 DeviceType=IOS
-BaseProfileName=iPhone15
+BaseProfileName=IOS_Epic
++CVars=ios.PhysicalScreenDensity=460
 
 [iPhone16e DeviceProfile]
 DeviceType=IOS
-BaseProfileName=iPhone15
+BaseProfileName=IOS_Epic
++CVars=ios.PhysicalScreenDensity=460
 
 [iPhone16Plus DeviceProfile]
 DeviceType=IOS
-BaseProfileName=iPhone15Plus
+BaseProfileName=IOS_Epic
 +CVars=ios.PhysicalScreenDensity=460
 
 [iPhone16Pro DeviceProfile]
 DeviceType=IOS
-BaseProfileName=iPhone15Pro
+BaseProfileName=IOS_Epic
++CVars=ios.PhysicalScreenDensity=460
 
 [iPhone16ProMax DeviceProfile]
 DeviceType=IOS
-BaseProfileName=iPhone15ProMax
+BaseProfileName=IOS_Epic
++CVars=ios.PhysicalScreenDensity=460
 
 [iPhone17Pro DeviceProfile]
 DeviceType=IOS
-BaseProfileName=iPhone15Pro
+BaseProfileName=IOS_Epic
++CVars=ios.PhysicalScreenDensity=460
 
 [iPhone17ProMax DeviceProfile]
 DeviceType=IOS
-BaseProfileName=iPhone15ProMax
+BaseProfileName=IOS_Epic
++CVars=ios.PhysicalScreenDensity=460
 
 [iPhone17 DeviceProfile]
 DeviceType=IOS
-BaseProfileName=iPhone15
+BaseProfileName=IOS_Epic
++CVars=ios.PhysicalScreenDensity=460
 
+;iPad
 [iPhoneAir DeviceProfile]
 DeviceType=IOS
-BaseProfileName=iPhone15Pro
+BaseProfileName=IOS_Epic
++CVars=ios.PhysicalScreenDensity=460
 
 [iPodTouch7 DeviceProfile]
 DeviceType=IOS
-BaseProfileName=iPhone7
+BaseProfileName=IOS_Mid
 +CVars=ios.PhysicalScreenDensity=326
-+CVars=r.MobileContentScaleFactor=2
 
 [iPadAir2 DeviceProfile]
 DeviceType=IOS
-BaseProfileName=IOS
+BaseProfileName=IOS_Mid
 +CVars=ios.PhysicalScreenDensity=264
-+CVars=sg.PostProcessQuality=3
-+CVars=sg.ShadowQuality=2
-+CVars=sg.GlobalIlluminationQuality=2
 
 [iPadMini4 DeviceProfile]
 DeviceType=IOS
-BaseProfileName=IOS
+BaseProfileName=IOS_Low
 +CVars=ios.PhysicalScreenDensity=401
-+CVars=sg.PostProcessQuality=2
 
 [iPadPro DeviceProfile]
 DeviceType=IOS
-BaseProfileName=IOS
+BaseProfileName=IOS_Mid
 +CVars=ios.PhysicalScreenDensity=264
 +CVars=r.MobileContentScaleFactor=1.5
 +CVars=g.TimeoutForBlockOnRenderFence=3000000
 +CVars=r.Decal.StencilSizeThreshold=0.1
 +CVars=r.MetalComputeParameterSize=1024
 +CVars=r.EarlyZPass=3
-+CVars=sg.ShadowQuality=3
-+CVars=sg.GlobalIlluminationQuality=3
-+CVars=sg.PostProcessQuality=3
-+CVars=sg.AntiAliasingQuality=3
 
 [iPadPro97 DeviceProfile]
 DeviceType=IOS
 BaseProfileName=IPadPro
-+CVars=ios.PhysicalScreenDensity=264
 
 [iPadPro129 DeviceProfile]
 DeviceType=IOS
 BaseProfileName=IPadPro
-+CVars=ios.PhysicalScreenDensity=264
 
 [iPadPro105 DeviceProfile]
 DeviceType=IOS
 BaseProfileName=IPadPro
-+CVars=ios.PhysicalScreenDensity=264
 
 [iPadPro11 DeviceProfile]
 DeviceType=IOS
 BaseProfileName=IPadPro
-+CVars=ios.PhysicalScreenDensity=264
 
 [iPadPro2_11 DeviceProfile]
 DeviceType=IOS
 BaseProfileName=IPadPro
-+CVars=ios.PhysicalScreenDensity=264
 
-[iPadPro3_11 DeviceProfile]
-DeviceType=IOS
-BaseProfileName=IPadPro
-+CVars=ios.PhysicalScreenDensity=264
-
-[iPadPro4_11 DeviceProfile]
-DeviceType=IOS
-BaseProfileName=IPadPro3_11
-
-[iPadPro5_11 DeviceProfile]
-DeviceType=IOS
-BaseProfileName=IPadPro3_11
-
 [iPad5 DeviceProfile]
 DeviceType=IOS
-BaseProfileName=IPadPro97
-+CVars=ios.PhysicalScreenDensity=264
+BaseProfileName=IPadPro
 
 [iPad6 DeviceProfile]
 DeviceType=IOS
-BaseProfileName=IPadPro97
-+CVars=ios.PhysicalScreenDensity=264
+BaseProfileName=IPadPro
 
 [iPad7 DeviceProfile]
 DeviceType=IOS
-BaseProfileName=IPadPro97
-+CVars=ios.PhysicalScreenDensity=264
+BaseProfileName=IPadPro
 
 [iPad8 DeviceProfile]
 DeviceType=IOS
-BaseProfileName=IPadPro97
-+CVars=ios.PhysicalScreenDensity=264
+BaseProfileName=IPadPro
 
 [iPad9 DeviceProfile]
 DeviceType=IOS
-BaseProfileName=IPadPro97
-+CVars=ios.PhysicalScreenDensity=264
+BaseProfileName=IPadPro
 
-[iPad10 DeviceProfile]
-DeviceType=IOS
-BaseProfileName=IPad9
-
 [iPadPro2_129 DeviceProfile]
 DeviceType=IOS
 BaseProfileName=IPadPro
-+CVars=ios.PhysicalScreenDensity=264
 
 [iPadPro3_129 DeviceProfile]
 DeviceType=IOS
 BaseProfileName=IPadPro
-+CVars=ios.PhysicalScreenDensity=264
 
 [iPadPro4_129 DeviceProfile]
 DeviceType=IOS
 BaseProfileName=IPadPro
-+CVars=ios.PhysicalScreenDensity=264
 
-[iPadPro5_129 DeviceProfile]
-DeviceType=IOS
-BaseProfileName=IPadPro
-+CVars=ios.PhysicalScreenDensity=264
-
-[iPadPro6_129 DeviceProfile]
-DeviceType=IOS
-BaseProfileName=IPadPro5_129
-
-[iPadPro7_129 DeviceProfile]
-DeviceType=IOS
-BaseProfileName=IPadPro5_129
-
 [iPadAir3 DeviceProfile]
 DeviceType=IOS
 BaseProfileName=IPadPro
+
+[iPadMini5 DeviceProfile]
+DeviceType=IOS
+BaseProfileName=IPadPro
++CVars=ios.PhysicalScreenDensity=401
+
+;iPad13
+[iPadPro3_11 DeviceProfile]
+DeviceType=IOS
+BaseProfileName=IOS_High
 +CVars=ios.PhysicalScreenDensity=264
++CVars=r.MobileContentScaleFactor=1.5
++CVars=g.TimeoutForBlockOnRenderFence=3000000
++CVars=r.Decal.StencilSizeThreshold=0.1
++CVars=r.MetalComputeParameterSize=1024
++CVars=r.EarlyZPass=3
 
 [iPadAir4 DeviceProfile]
 DeviceType=IOS
-BaseProfileName=IPadPro
-+CVars=ios.PhysicalScreenDensity=264
+BaseProfileName=iPadPro3_11
 
 [iPadAir5 DeviceProfile]
 DeviceType=IOS
-BaseProfileName=IPadAir4
+BaseProfileName=iPadPro3_11
 
+[iPad10 DeviceProfile]
+DeviceType=IOS
+BaseProfileName=iPadPro3_11
+
+[iPadPro5_129 DeviceProfile]
+DeviceType=IOS
+BaseProfileName=iPadPro3_11
+
+;iPad14
+[iPadPro4_11 DeviceProfile]
+DeviceType=IOS
+BaseProfileName=IOS_Epic
++CVars=ios.PhysicalScreenDensity=264
++CVars=r.MobileContentScaleFactor=1.5
++CVars=g.TimeoutForBlockOnRenderFence=3000000
++CVars=r.Decal.StencilSizeThreshold=0.1
++CVars=r.MetalComputeParameterSize=1024
++CVars=r.EarlyZPass=3
+
+[iPadPro6_129 DeviceProfile]
+DeviceType=IOS
+BaseProfileName=iPadPro4_11
+
 [iPadAir6 DeviceProfile]
 DeviceType=IOS
-BaseProfileName=IPadAir4
+BaseProfileName=iPadPro4_11
 
 [iPadAir_13 DeviceProfile]
 DeviceType=IOS
-BaseProfileName=iPadAir5
-+CVars=ios.PhysicalScreenDensity=264
+BaseProfileName=iPadPro4_11
 
+[iPadMini6 DeviceProfile]
+DeviceType=IOS
+BaseProfileName=iPadPro4_11
++CVars=ios.PhysicalScreenDensity=326
+
+;iPad15
 [iPadAir7 DeviceProfile]
 DeviceType=IOS
-BaseProfileName=iPadAir6
+BaseProfileName=iPadPro4_11
 
 [iPadAir7_13 DeviceProfile]
 DeviceType=IOS
-BaseProfileName=iPadAir_13
+BaseProfileName=iPadPro4_11
 
-[iPadMini5 DeviceProfile]
+;iPad16
+[iPadPro5_11 DeviceProfile]
 DeviceType=IOS
-BaseProfileName=IPadPro
-+CVars=ios.PhysicalScreenDensity=401
+BaseProfileName=iPadPro4_11
 
-[iPadMini6 DeviceProfile]
+[iPadPro7_129 DeviceProfile]
 DeviceType=IOS
-BaseProfileName=IPadPro
-+CVars=ios.PhysicalScreenDensity=326
+BaseProfileName=iPadPro4_11
 
 [iPadMini7 DeviceProfile]
 DeviceType=IOS
-BaseProfileName=iPadMini6
+BaseProfileName=iPadPro4_11
 
 [AppleTV DeviceProfile]
 DeviceType=IOS
-BaseProfileName=IOS
+BaseProfileName=IOS_Mid
 +CVars=ios.PhysicalScreenDensity=0
-+CVars=sg.PostProcessQuality=3
-+CVars=sg.ShadowQuality=3
-+CVars=sg.GlobalIlluminationQuality=3
 +CVars=UI.SlateSDFText.RasterizationMode=SdfApproximation
 +CVars=UI.SlateSDFText.ResolutionLevel=1
 
@@ -786,7 +817,7 @@
 
 [VisionPro DeviceProfile]
 DeviceType=VisionOS
-BaseProfileName=IOS
+BaseProfileName=IOS_High
 +CVars=r.Mobile.XRMSAAMode=2
 +CVars=xr.OpenXRLateUpdateDeviceLocationsAfterReflections=1
 
@@ -911,6 +942,12 @@
 ; Optimize Slate signed distance field text for performance
 +CVars=UI.SlateSDFText.RasterizationMode=SdfApproximation
 +CVars=UI.SlateSDFText.ResolutionLevel=1
+; Use 'OcclusionFeedback' instead of occlusion queires to reduce drawcalls
++CVars=r.OcclusionFeedback.Enable=1
+; Max 256k partices
++CVars=fx.GPUSimulationTextureSizeX=512
++CVars=fx.GPUSimulationTextureSizeY=512
++CVars=r.Streaming.MaxTempMemoryAllowed=25
 ; PreviewAllowlistCVars and PreviewDenyListCVars are arrays of cvars (or prefixes) that are included or excluded from being applied in editor preview.
 ; If any PreviewAllowlistCVars is set, cvars are denied by default.
 PreviewAllowlistCVars=sg.
@@ -975,6 +1012,23 @@
 +CVars=UI.SlateSDFText.RasterizationMode=Msdf
 +CVars=UI.SlateSDFText.ResolutionLevel=2
 
+[Android_Epic DeviceProfile]
+DeviceType=Android
+BaseProfileName=Android
++CVars=sg.ViewDistanceQuality=3
++CVars=sg.AntiAliasingQuality=3
++CVars=sg.ShadowQuality=3
++CVars=sg.GlobalIlluminationQuality=3
++CVars=sg.ReflectionQuality=3
++CVars=sg.PostProcessQuality=3
++CVars=sg.TextureQuality=3
++CVars=sg.EffectsQuality=3
++CVars=sg.FoliageQuality=3
++CVars=sg.LandscapeQuality=3
++CVars=r.MobileContentScaleFactor=1.0
++CVars=UI.SlateSDFText.RasterizationMode=Msdf
++CVars=UI.SlateSDFText.ResolutionLevel=2
+
 [Android_Emulator DeviceProfile]
 DeviceType=Android
 BaseProfileName=Android_High
@@ -1075,7 +1129,7 @@
 
 [Android_Adreno7xx DeviceProfile]
 DeviceType=Android
-BaseProfileName=Android_High
+BaseProfileName=Android_Epic
 
 [Android_Adreno7xx_Vulkan DeviceProfile]
 DeviceType=Android
@@ -1085,7 +1139,7 @@
 
 [Android_Adreno8xx DeviceProfile]
 DeviceType=Android
-BaseProfileName=Android_High
+BaseProfileName=Android_Epic
 
 [Android_Adreno8xx_Vulkan DeviceProfile]
 DeviceType=Android
@@ -1108,7 +1162,7 @@
 
 [Android_Xclipse_9xx DeviceProfile]
 DeviceType=Android
-BaseProfileName=Android_High
+BaseProfileName=Android_Epic
 
 [Android_Xclipse_9xx_Vulkan DeviceProfile]
 DeviceType=Android
@@ -1120,7 +1174,7 @@
 
 [Android_Xclipse_xxx DeviceProfile]
 DeviceType=Android
-BaseProfileName=Android_Mid
+BaseProfileName=Android_Epic
 
 [Android_Xclipse_xxx_Vulkan DeviceProfile]
 DeviceType=Android
@@ -1212,7 +1266,7 @@
 
 [Android_Mali_G72 DeviceProfile]
 DeviceType=Android
-BaseProfileName=Android_High
+BaseProfileName=Android_Mid
 
 [Android_Mali_G72_Vulkan DeviceProfile]
 DeviceType=Android
@@ -1252,7 +1306,7 @@
 
 [Android_Mali_G710 DeviceProfile]
 DeviceType=Android
-BaseProfileName=Android_High
+BaseProfileName=Android_Epic
 
 [Android_Mali_G710_Vulkan DeviceProfile]
 DeviceType=Android
diff --git a/Engine/Config/IOS/BaseIOSEngine.ini b/Engine/Config/IOS/BaseIOSEngine.ini
--- a/Engine/Config/IOS/BaseIOSEngine.ini
+++ b/Engine/Config/IOS/BaseIOSEngine.ini
@@ -15,6 +15,7 @@ r.TemporalAA.Mobile.UseCompute=0
 r.Visibility.DynamicMeshElements.Parallel=0
 ; Volumetric Clouds are not supported by mobile renderer, skip materials/global shaders compilation.
 r.VolumetricCloud.Support=0
+r.OIT.SortedTriangles=0
 
 [/Script/UdpMessaging.UdpMessagingSettings]
 ; IOS requires an additional entitlement (com.apple.developer.networking.multicast) to enable UDP Local Networking
diff --git a/Engine/Config/IOS/IOSScalability.ini b/Engine/Config/IOS/IOSScalability.ini
--- a/Engine/Config/IOS/IOSScalability.ini
+++ b/Engine/Config/IOS/IOSScalability.ini
@@ -53,6 +53,7 @@ r.Shadow.CSM.MaxCascades=1
 r.Shadow.RadiusThreshold=0.06
 r.VolumetricFog=0
 r.ContactShadows=0
+r.DistanceFieldShadowing=0
 r.DFShadowQuality=0
 
 [ShadowQuality@1]
@@ -63,6 +64,7 @@ r.Shadow.CSM.MaxCascades=1
 r.Shadow.RadiusThreshold=0.06
 r.VolumetricFog=0
 r.ContactShadows=0
+r.DistanceFieldShadowing=0
 r.DFShadowQuality=0
 
 [ShadowQuality@2]
@@ -73,7 +75,7 @@ r.Shadow.CSM.MaxCascades=2
 r.Shadow.RadiusThreshold=0.06
 r.VolumetricFog=0
 r.ContactShadows=0
-r.DFShadowQuality=0
+r.DFShadowQuality=1
 r.Mobile.SupportInsetShadows=True
 
 [ShadowQuality@3]
@@ -145,15 +147,19 @@ r.Mobile.AmbientOcclusionQuality=3
 r.LensFlareQuality=2
 
 [AntiAliasingQuality@0]
+r.Mobile.AntiAliasing=0
 r.FXAA.Quality=0
 
 [AntiAliasingQuality@1]
+r.Mobile.AntiAliasing=1
 r.FXAA.Quality=0
 
 [AntiAliasingQuality@2]
+r.Mobile.AntiAliasing=2
 r.FXAA.Quality=0
 
 [AntiAliasingQuality@3]
+r.Mobile.AntiAliasing=2
 r.FXAA.Quality=0
 
 ; ------------------------------------------------------------------------------
diff --git a/Engine/Source/Runtime/Renderer/Private/ShadowDepthRendering.cpp b/Engine/Source/Runtime/Renderer/Private/ShadowDepthRendering.cpp
--- a/Engine/Source/Runtime/Renderer/Private/ShadowDepthRendering.cpp
+++ b/Engine/Source/Runtime/Renderer/Private/ShadowDepthRendering.cpp
@@ -126,12 +126,6 @@ static TAutoConsoleVariable<int32> CVarShouldBeginDeferredCullingAfterShadowRend
 	TEXT("Temporary: If turned on, a new deferred scope of instance culling will be started directly after shadow depth rendering (iff instance occlusion culling is enabled) to work around a bug."),
 	ECVF_RenderThreadSafe);
 
-static TAutoConsoleVariable<int32> CVarShadowForceSerialSingleRenderPass(
-	TEXT("r.Shadow.ForceSerialSingleRenderPass"),
-	0,
-	TEXT("Force Serial shadow passes to render in 1 pass."),
-	ECVF_RenderThreadSafe);
-
 TAutoConsoleVariable<int32> CVarNaniteShadows(
 	TEXT("r.Shadow.Nanite"),
 	1,
