Added engine device profiles for Mali-G1-Ultra and similar devices

diff --git a/Engine/Config/BaseDeviceProfiles.ini b/Engine/Config/BaseDeviceProfiles.ini
--- a/Engine/Config/BaseDeviceProfiles.ini
+++ b/Engine/Config/BaseDeviceProfiles.ini
@@ -142,6 +142,10 @@
 +DeviceProfileNameAndTypes=Android_Mali_G7xx_Vulkan,Android
 +DeviceProfileNameAndTypes=Android_Mali_G9xx,Android
 +DeviceProfileNameAndTypes=Android_Mali_G9xx_Vulkan,Android
++DeviceProfileNameAndTypes=Android_Mali_Gx,Android
++DeviceProfileNameAndTypes=Android_Mali_Gx_Vulkan,Android
++DeviceProfileNameAndTypes=Android_Mali_Gx_Ultra,Android
++DeviceProfileNameAndTypes=Android_Mali_Gx_Ultra_Vulkan,Android
 +DeviceProfileNameAndTypes=Android_Xclipse_5xx,Android
 +DeviceProfileNameAndTypes=Android_Xclipse_5xx_Vulkan,Android
 +DeviceProfileNameAndTypes=Android_Xclipse_9xx,Android
@@ -890,6 +894,10 @@ BaseProfileName=VisionPro
 +MatchProfile=(Profile="Android_Mali_G7xx",Match=((SourceType=SRC_GpuFamily,CompareType=CMP_Regex,MatchString="^Mali\\-G7[0-9][0-9]")))
 +MatchProfile=(Profile="Android_Mali_G9xx_Vulkan",Match=((SourceType=SRC_GpuFamily,CompareType=CMP_Regex,MatchString="^Mali\\-G9[0-9][0-9]"),(SourceType=SRC_VulkanAvailable,CompareType=CMP_Equal,MatchString="true")))
 +MatchProfile=(Profile="Android_Mali_G9xx",Match=((SourceType=SRC_GpuFamily,CompareType=CMP_Regex,MatchString="^Mali\\-G9[0-9][0-9]")))
++MatchProfile=(Profile="Android_Mali_Gx_Ultra_Vulkan",Match=((SourceType=SRC_GpuFamily,CompareType=CMP_Regex,MatchString="^Mali\\-G[0-9]\\-Ultra"),(SourceType=SRC_VulkanAvailable,CompareType=CMP_Equal,MatchString="true")))
++MatchProfile=(Profile="Android_Mali_Gx_Ultra",Match=((SourceType=SRC_GpuFamily,CompareType=CMP_Regex,MatchString="^Mali\\-G[0-9]\\-Ultra")))
++MatchProfile=(Profile="Android_Mali_Gx_Vulkan",Match=((SourceType=SRC_GpuFamily,CompareType=CMP_Regex,MatchString="^Mali\\-G[0-9]\\-"),(SourceType=SRC_VulkanAvailable,CompareType=CMP_Equal,MatchString="true")))
++MatchProfile=(Profile="Android_Mali_Gx",Match=((SourceType=SRC_GpuFamily,CompareType=CMP_Regex,MatchString="^Mali\\-G[0-9]\\-")))
 
 ; Enable Vulkan on any unknown Mali running Android 10 or later
 +MatchProfile=(Profile="Android_Unknown_Vulkan",Match=((SourceType=SRC_GpuFamily,CompareType=CMP_Regex,MatchString="^Mali"),(SourceType=SRC_AndroidVersion, CompareType=CMP_Regex,MatchString="([0-9]+).*"),(SourceType=SRC_PreviousRegexMatch,CompareType=CMP_GreaterEqual,MatchString="10"),(SourceType=SRC_VulkanAvailable,CompareType=CMP_Equal,MatchString="true")))
@@ -1334,6 +1342,24 @@ BaseProfileName=Android_Mali_G9xx
 ; enable Vulkan
 +CVars=r.Android.DisableVulkanSupport=0
 
+[Android_Mali_Gx DeviceProfile]
+DeviceType=Android
+BaseProfileName=Android_High
+
+[Android_Mali_Gx_Vulkan DeviceProfile]
+DeviceType=Android
+BaseProfileName=Android_Mali_Gx
++CVars=r.Android.DisableVulkanSupport=0
+
+[Android_Mali_Gx_Ultra DeviceProfile]
+DeviceType=Android
+BaseProfileName=Android_Epic
+
+[Android_Mali_Gx_Ultra_Vulkan DeviceProfile]
+DeviceType=Android
+BaseProfileName=Android_Mali_Gx_Ultra
++CVars=r.Android.DisableVulkanSupport=0
+
 [Android_TegraK1 DeviceProfile]
 DeviceType=Android
 BaseProfileName=Android_High
