Disable occlusion feedback (r.OcclusionFeedback.Enable) for Quest devices.
Occlusion feedback was recently enabled by default on mobile devices. It's currently visually bugged on Quest, and disables the preferred Adreno Occlusion Fast Mode.

diff --git a/Engine/Config/BaseDeviceProfiles.ini b/Engine/Config/BaseDeviceProfiles.ini
--- a/Engine/Config/BaseDeviceProfiles.ini
+++ b/Engine/Config/BaseDeviceProfiles.ini
@@ -1385,18 +1385,21 @@ BaseProfileName=Android_Mid
 DeviceType=Android
 BaseProfileName=Android_OpenXR
 +CVars=r.Mobile.Oculus.ForceSymmetric=1
++Cvars=r.OcclusionFeedback.Enable=0
 +CVars=r.Mobile.AdrenoOcclusionMode=1
 
 [Meta_Quest_Pro DeviceProfile]
 DeviceType=Android
 BaseProfileName=Android_OpenXR
 +CVars=r.Mobile.Oculus.ForceSymmetric=1
++Cvars=r.OcclusionFeedback.Enable=0
 +CVars=r.Mobile.AdrenoOcclusionMode=1
 
 [Oculus_Quest2 DeviceProfile]
 DeviceType=Android
 BaseProfileName=Android_OpenXR
 +CVars=r.Mobile.Oculus.ForceSymmetric=1
++Cvars=r.OcclusionFeedback.Enable=0
 +CVars=r.Mobile.AdrenoOcclusionMode=1
 +CVars=sg.ShadowQuality=1
 
