Disable Vulkan support for PowerVR Series A-D for now due to driver issues

#android

diff --git a/Engine/Config/BaseDeviceProfiles.ini b/Engine/Config/BaseDeviceProfiles.ini
--- a/Engine/Config/BaseDeviceProfiles.ini
+++ b/Engine/Config/BaseDeviceProfiles.ini
@@ -917,9 +917,9 @@ BaseProfileName=VisionPro
 +MatchProfile=(Profile="Android_PowerVR_GE8xxx",Match=((SourceType=SRC_GpuFamily,CompareType=CMP_Regex,MatchString="PowerVR Rogue GE8[0-9]+")))
 +MatchProfile=(Profile="Android_PowerVR_GM9xxx_Vulkan",Match=((SourceType=SRC_GpuFamily,CompareType=CMP_Regex,MatchString="PowerVR Rogue GM9[0-9]+"),(SourceType=SRC_VulkanAvailable,CompareType=CMP_Equal,MatchString="true")))
 +MatchProfile=(Profile="Android_PowerVR_GM9xxx",Match=((SourceType=SRC_GpuFamily,CompareType=CMP_Regex,MatchString="PowerVR Rogue GM9[0-9]+")))
-+MatchProfile=(Profile="Android_PowerVR_Series_High_Vulkan",Match=((SourceType=SRC_GpuFamily,CompareType=CMP_Regex,MatchString="PowerVR [A-D]\\-Series [A-D]X[A-Z]+\\-([0-9]+)\\-[0-9]+"),(SourceType=SRC_PreviousRegexMatch,CompareType=CMP_GreaterEqual,MatchString="32"),(SourceType=SRC_VulkanAvailable,CompareType=CMP_Equal,MatchString="true")))
-+MatchProfile=(Profile="Android_PowerVR_Series_Mid_Vulkan",Match=((SourceType=SRC_GpuFamily,CompareType=CMP_Regex,MatchString="PowerVR [A-D]\\-Series [A-D]X[A-Z]+\\-([0-9]+)\\-[0-9]+"),(SourceType=SRC_PreviousRegexMatch,CompareType=CMP_GreaterEqual,MatchString="16"),(SourceType=SRC_VulkanAvailable,CompareType=CMP_Equal,MatchString="true")))
-+MatchProfile=(Profile="Android_PowerVR_Series_Low_Vulkan",Match=((SourceType=SRC_GpuFamily,CompareType=CMP_Regex,MatchString="PowerVR [A-D]\\-Series [A-D]X[A-Z]+\\-[0-9]+\\-[0-9]+"),(SourceType=SRC_VulkanAvailable,CompareType=CMP_Equal,MatchString="true")))
+;+MatchProfile=(Profile="Android_PowerVR_Series_High_Vulkan",Match=((SourceType=SRC_GpuFamily,CompareType=CMP_Regex,MatchString="PowerVR [A-D]\\-Series [A-D]X[A-Z]+\\-([0-9]+)\\-[0-9]+"),(SourceType=SRC_PreviousRegexMatch,CompareType=CMP_GreaterEqual,MatchString="32"),(SourceType=SRC_VulkanAvailable,CompareType=CMP_Equal,MatchString="true")))
+;+MatchProfile=(Profile="Android_PowerVR_Series_Mid_Vulkan",Match=((SourceType=SRC_GpuFamily,CompareType=CMP_Regex,MatchString="PowerVR [A-D]\\-Series [A-D]X[A-Z]+\\-([0-9]+)\\-[0-9]+"),(SourceType=SRC_PreviousRegexMatch,CompareType=CMP_GreaterEqual,MatchString="16"),(SourceType=SRC_VulkanAvailable,CompareType=CMP_Equal,MatchString="true")))
+;+MatchProfile=(Profile="Android_PowerVR_Series_Low_Vulkan",Match=((SourceType=SRC_GpuFamily,CompareType=CMP_Regex,MatchString="PowerVR [A-D]\\-Series [A-D]X[A-Z]+\\-[0-9]+\\-[0-9]+"),(SourceType=SRC_VulkanAvailable,CompareType=CMP_Equal,MatchString="true")))
 +MatchProfile=(Profile="Android_PowerVR_Series_High",Match=((SourceType=SRC_GpuFamily,CompareType=CMP_Regex,MatchString="PowerVR [A-D]\\-Series [A-D]X[A-Z]+\\-([0-9]+)\\-[0-9]+"),(SourceType=SRC_PreviousRegexMatch,CompareType=CMP_GreaterEqual,MatchString="32")))
 +MatchProfile=(Profile="Android_PowerVR_Series_Mid",Match=((SourceType=SRC_GpuFamily,CompareType=CMP_Regex,MatchString="PowerVR [A-D]\\-Series [A-D]X[A-Z]+\\-([0-9]+)\\-[0-9]+"),(SourceType=SRC_PreviousRegexMatch,CompareType=CMP_GreaterEqual,MatchString="16")))
 +MatchProfile=(Profile="Android_PowerVR_Series_Low",Match=((SourceType=SRC_GpuFamily,CompareType=CMP_Regex,MatchString="PowerVR [A-D]\\-Series [A-D]X[A-Z]+\\-[0-9]+\\-[0-9]+")))
