Lumen - remove r.Lumen.ScreenProbeGather.IntegrateDownsampleFactor 2 from High scalability as it was generating too much noise and softening normals. This has to be now manually tweaked per title for a more explicit tradeoff.

diff --git a/Engine/Config/BaseScalability.ini b/Engine/Config/BaseScalability.ini
--- a/Engine/Config/BaseScalability.ini
+++ b/Engine/Config/BaseScalability.ini
@@ -297,7 +297,6 @@ r.Lumen.ScreenProbeGather.NumAdaptiveProbes=16
 r.Lumen.ScreenProbeGather.TracingOctahedronResolution=8
 r.Lumen.ScreenProbeGather.IrradianceFormat=1
 r.Lumen.ScreenProbeGather.StochasticInterpolation=1
-r.Lumen.ScreenProbeGather.IntegrateDownsampleFactor=2
 r.Lumen.ScreenProbeGather.FullResolutionJitterWidth=1
 r.Lumen.ScreenProbeGather.TwoSidedFoliageBackfaceDiffuse=0
 r.Lumen.ScreenProbeGather.ScreenTraces.HZBTraversal.FullResDepth=1
@@ -329,7 +328,6 @@ r.Lumen.ScreenProbeGather.NumAdaptiveProbes=8
 r.Lumen.ScreenProbeGather.TracingOctahedronResolution=8
 r.Lumen.ScreenProbeGather.IrradianceFormat=0
 r.Lumen.ScreenProbeGather.StochasticInterpolation=0
-r.Lumen.ScreenProbeGather.IntegrateDownsampleFactor=1
 r.Lumen.ScreenProbeGather.FullResolutionJitterWidth=1
 r.Lumen.ScreenProbeGather.TwoSidedFoliageBackfaceDiffuse=1
 r.Lumen.ScreenProbeGather.ScreenTraces.HZBTraversal.FullResDepth=1
@@ -362,7 +360,6 @@ r.Lumen.ScreenProbeGather.NumAdaptiveProbes=4
 r.Lumen.ScreenProbeGather.TracingOctahedronResolution=16
 r.Lumen.ScreenProbeGather.IrradianceFormat=0
 r.Lumen.ScreenProbeGather.StochasticInterpolation=0
-r.Lumen.ScreenProbeGather.IntegrateDownsampleFactor=1
 r.Lumen.ScreenProbeGather.FullResolutionJitterWidth=.5
 r.Lumen.ScreenProbeGather.TwoSidedFoliageBackfaceDiffuse=1
 r.Lumen.ScreenProbeGather.ScreenTraces.HZBTraversal.FullResDepth=1
