Update Help->Report a bug url.

diff --git a/Engine/Config/BaseEditor.ini b/Engine/Config/BaseEditor.ini
--- a/Engine/Config/BaseEditor.ini
+++ b/Engine/Config/BaseEditor.ini
@@ -322,7 +322,7 @@ SourceCodeIDEURL_Other="http://en.wikipedia.org/wiki/Integrated_development_envi
 EpicGamesURL="http://UnrealEngine.com"
 ViewportControlsURL="viewport-controls-in-unreal-engine"
 MeshSimplificationPluginsURL="/ue/marketplace/content-cat/assets/codeplugins"
-ReportABugURL="https://epicsupport.force.com/unrealengine/s/"
+ReportABugURL="https://unrealbugsubmissions.epicgames.com/"
 IssueTrackerURL="https://issues.unrealengine.com/"
 
 
