- Fixed a couple of bad comments/logs

diff --git a/Engine/Config/BaseEngine.ini b/Engine/Config/BaseEngine.ini
--- a/Engine/Config/BaseEngine.ini
+++ b/Engine/Config/BaseEngine.ini
@@ -351,7 +351,7 @@ ClassName=/Script/OnlineSubsystemUtils.OnlineEngineInterfaceImpl
 [CoreRedirects]
 
 ; All redirects prior to UE 4.10 are now in FCoreRedirects::RegisterNativeRedirects, new ones should be added to the bottom of this section
-; Redirects can now go in any ini file, including Plugin inis like PluginName/Config/BasePluginName.ini
+; Redirects can now go in any ini file, including Plugin inis like PluginName/Config/DefaultPluginName.ini
 ; This redirect format replaces the old ActiveClassRedirects/K2FieldRedirects format, and the syntax is the same for every type
 ; The -DebugCoreRedirects command line parameter can be used to diagnose issues when adding new redirects
 ;
diff --git a/Engine/Source/Runtime/Engine/Private/UnrealEngine.cpp b/Engine/Source/Runtime/Engine/Private/UnrealEngine.cpp
--- a/Engine/Source/Runtime/Engine/Private/UnrealEngine.cpp
+++ b/Engine/Source/Runtime/Engine/Private/UnrealEngine.cpp
@@ -10615,7 +10615,7 @@ bool UEngine::HandleGetIniCommand(const TCHAR* Cmd, FOutputDevice& Ar)
 						ConfigSystem->GetArray(*SectionName, KeyName, FinalValues, IniName);
 						
 						Ar.Logf(TEXT(" "));
-						Ar.Logf(TEXT("Final Value(s) [should match ImMemory above!]"));
+						Ar.Logf(TEXT("Final Value(s) [should match InMemory above!]"));
 						Ar.Logf(TEXT("[%s]"), *SectionName);
 						
 						if (FinalValues.Num() == 1)
