[Editor TRS] Plugin enabled by default, actual gizmo isn't used unless editor setting toggled

#was-backout 46683522

diff --git a/Engine/Config/BaseGame.ini b/Engine/Config/BaseGame.ini
--- a/Engine/Config/BaseGame.ini
+++ b/Engine/Config/BaseGame.ini
@@ -312,6 +312,7 @@ MapMode=Uncooked
 +DisabledPlugins=AssetSearch
 ; various plugins that cause issues in cooked/shipping editors (unknown reasons for now)
 +DisabledPlugins=ObjectBrowser
++DisabledPlugins=EditorTRS
 ; stuff that doesn't work well or is not needed in shipping
 +DisabledPlugins=FunctionalTestingEditor
 +DisabledPlugins=RuntimeTests
