Add WriteOnly setting for HttpCacheStore.  Refactor accompanying code to specify a setting in either config, env var, or commandline parameter.

diff --git a/Engine/Config/BaseEngine.ini b/Engine/Config/BaseEngine.ini
--- a/Engine/Config/BaseEngine.ini
+++ b/Engine/Config/BaseEngine.ini
@@ -2684,7 +2684,7 @@ Port=8558
 ;   CommandLineOAuthSecretOverride: The name of a command line switch that can override the OAuth secret, e.g., CloudDataCacheOAuthSecret. Optional.
 ;   OAuthScope: The scope to use with the OAuth provider. Default is cache_access.
 ;   OAuthProviderIdentifier: The provider ID when using OIDC. Optional.
-;   OAuthAccessTokenEnvOverride: The name of a command line switch that can provide the OAuth access token to use, e.g., UE-CloudDataCacheAccessToken. Optional.
+;   EnvOAuthAccessTokenOverride: The name of a command line switch that can provide the OAuth access token to use, e.g., UE-CloudDataCacheAccessToken. Optional.
 ;   OAuthPinnedPublicKeys: A semicolon-delimited list of pinned public keys for the OAuth provider. Optional.
 ;   ReadOnly: Whether the cache is treated as read-only. Default is false.
 ;   HttpVersion: An override for which http version to use
@@ -2744,7 +2744,7 @@ Local=(Type=FileSystem, DeleteOnly=true, UnusedFileAge=8, Path=%ENGINEDIR%Derive
 ; Configure the shared cache that is accessed after local.
 Shared=(Type=FileSystem, UnusedFileAge=10, FoldersToClean=10, ConsiderSlowAt=70, Path=?EpicDDC, EnvPathOverride=UE-SharedDataCachePath, EditorOverrideSetting=SharedDerivedDataCache, CommandLineOverride=SharedDataCachePath)
 ; Configure the cloud cache that is accessed after shared.
-Cloud=(Type=Cloud, ServerID=Cloud, EnvHostOverride=UE-CloudDataCacheHost, CommandLineHostOverride=CloudDataCacheHost, CommandLineOAuthSecretOverride="CloudDataCacheOAuthSecret", OAuthAccessTokenEnvOverride="UE-CloudDataCacheAccessToken", EnvHttpVersionOverride="UE-CloudDataCacheHttpVersion", CommandLineHttpVersionOverride="CloudDataCacheHttpVersion", EnvOAuthProviderIdentifierOverride="UE-CloudDataCacheOAuthProviderIdentifier", CommandLineOAuthProviderIdentifierOverride="CloudDataCacheOAuthProviderIdentifier")
+Cloud=(Type=Cloud, ServerID=Cloud, EnvHostOverride=UE-CloudDataCacheHost, CommandLineHostOverride=CloudDataCacheHost, CommandLineOAuthSecretOverride="CloudDataCacheOAuthSecret", EnvOAuthAccessTokenOverride="UE-CloudDataCacheAccessToken", EnvHttpVersionOverride="UE-CloudDataCacheHttpVersion", CommandLineHttpVersionOverride="CloudDataCacheHttpVersion", EnvOAuthProviderIdentifierOverride="UE-CloudDataCacheOAuthProviderIdentifier", CommandLineOAuthProviderIdentifierOverride="CloudDataCacheOAuthProviderIdentifier")
 ; Configure a Project Pak node. This is a pre-generated DDC data file for the project that can be distributed to reduce runtime fetches/generation
 ; See documentation for how to create a DDP via the DerivedDataCache commandlet
 Pak=(Type=ReadPak, Filename="%GAMEDIR%DerivedDataCache/DDC.ddp")
@@ -2759,7 +2759,7 @@ AsyncPut=(Type=AsyncPut, Inner=Hierarchy)
 Hierarchy=(Type=Hierarchical, Inner=Pak, Inner=CompressedPak, Inner=EnginePak, Inner=EnterprisePak, Inner=ZenLocal, Inner=Local, Inner=ZenShared, Inner=Shared, Inner=Cloud)
 Local=(Type=FileSystem, DeleteOnly=true, UnusedFileAge=8, Path="%ENGINEVERSIONAGNOSTICUSERDIR%DerivedDataCache", EnvPathOverride=UE-LocalDataCachePath, EditorOverrideSetting=LocalDerivedDataCache)
 Shared=(Type=FileSystem, UnusedFileAge=10, FoldersToClean=10, ConsiderSlowAt=70, Path=?EpicDDC, EnvPathOverride=UE-SharedDataCachePath, EditorOverrideSetting=SharedDerivedDataCache, CommandLineOverride=SharedDataCachePath)
-Cloud=(Type=Cloud, ServerID=Cloud, EnvHostOverride=UE-CloudDataCacheHost, CommandLineHostOverride=CloudDataCacheHost, CommandLineOAuthSecretOverride="CloudDataCacheOAuthSecret", OAuthAccessTokenEnvOverride="UE-CloudDataCacheAccessToken", EnvHttpVersionOverride="UE-CloudDataCacheHttpVersion", CommandLineHttpVersionOverride="CloudDataCacheHttpVersion", EnvOAuthProviderIdentifierOverride="UE-CloudDataCacheOAuthProviderIdentifier", CommandLineOAuthProviderIdentifierOverride="CloudDataCacheOAuthProviderIdentifier")
+Cloud=(Type=Cloud, ServerID=Cloud, EnvHostOverride=UE-CloudDataCacheHost, CommandLineHostOverride=CloudDataCacheHost, CommandLineOAuthSecretOverride="CloudDataCacheOAuthSecret", EnvOAuthAccessTokenOverride="UE-CloudDataCacheAccessToken", EnvHttpVersionOverride="UE-CloudDataCacheHttpVersion", CommandLineHttpVersionOverride="CloudDataCacheHttpVersion", EnvOAuthProviderIdentifierOverride="UE-CloudDataCacheOAuthProviderIdentifier", CommandLineOAuthProviderIdentifierOverride="CloudDataCacheOAuthProviderIdentifier")
 Pak=(Type=ReadPak, Filename="%GAMEDIR%DerivedDataCache/DDC.ddp")
 CompressedPak=(Type=ReadPak, Filename="%GAMEDIR%DerivedDataCache/Compressed.ddp", Compressed=true)
 EnginePak=(Type=ReadPak, Filename=../../../Engine/DerivedDataCache/Compressed.ddp, Compressed=true)
@@ -2771,7 +2771,7 @@ ZenShared=(Type=Zen, ServerID=Shared)
 Root=(Type=Hierarchical, Inner=Pak, Inner=EnginePak, Inner=Local, Inner=ZenShared, Inner=Shared, Inner=Cloud)
 Local=(Type=FileSystem, UnusedFileAge=34, PromptIfMissing=true, Path=%ENGINEDIR%DerivedDataCache, EnvPathOverride=UE-LocalDataCachePath, EditorOverrideSetting=LocalDerivedDataCache, CommandLineOverride=LocalDataCachePath)
 Shared=(Type=FileSystem, UnusedFileAge=10, FoldersToClean=10, ConsiderSlowAt=70, Path=?EpicDDC, EnvPathOverride=UE-SharedDataCachePath, EditorOverrideSetting=SharedDerivedDataCache, CommandLineOverride=SharedDataCachePath)
-Cloud=(Type=Cloud, ServerID=Cloud, EnvHostOverride=UE-CloudDataCacheHost, CommandLineHostOverride=CloudDataCacheHost, CommandLineOAuthSecretOverride="CloudDataCacheOAuthSecret", OAuthAccessTokenEnvOverride="UE-CloudDataCacheAccessToken", EnvHttpVersionOverride="UE-CloudDataCacheHttpVersion", CommandLineHttpVersionOverride="CloudDataCacheHttpVersion", EnvOAuthProviderIdentifierOverride="UE-CloudDataCacheOAuthProviderIdentifier", CommandLineOAuthProviderIdentifierOverride="CloudDataCacheOAuthProviderIdentifier")
+Cloud=(Type=Cloud, ServerID=Cloud, EnvHostOverride=UE-CloudDataCacheHost, CommandLineHostOverride=CloudDataCacheHost, CommandLineOAuthSecretOverride="CloudDataCacheOAuthSecret", EnvOAuthAccessTokenOverride="UE-CloudDataCacheAccessToken", EnvHttpVersionOverride="UE-CloudDataCacheHttpVersion", CommandLineHttpVersionOverride="CloudDataCacheHttpVersion", EnvOAuthProviderIdentifierOverride="UE-CloudDataCacheOAuthProviderIdentifier", CommandLineOAuthProviderIdentifierOverride="CloudDataCacheOAuthProviderIdentifier")
 Pak=(Type=ReadPak, Filename="%GAMEDIR%DerivedDataCache/DDC.ddp")
 EnginePak=(Type=ReadPak, Filename="%ENGINEDIR%DerivedDataCache/DDC.ddp")
 ZenShared=(Type=Zen, ServerID=Shared)
@@ -2780,7 +2780,7 @@ ZenShared=(Type=Zen, ServerID=Shared)
 Root=(Type=Hierarchical, Inner=Pak, Inner=CompressedPak, Inner=EnginePak, Inner=EnterprisePak, Inner=Local, Inner=ZenShared, Inner=Shared, Inner=Cloud)
 Local=(Type=FileSystem, UnusedFileAge=34, Path="%ENGINEVERSIONAGNOSTICUSERDIR%DerivedDataCache", EnvPathOverride=UE-LocalDataCachePath, EditorOverrideSetting=LocalDerivedDataCache)
 Shared=(Type=FileSystem, UnusedFileAge=10, FoldersToClean=10, ConsiderSlowAt=70, Path=?EpicDDC, EnvPathOverride=UE-SharedDataCachePath, EditorOverrideSetting=SharedDerivedDataCache, CommandLineOverride=SharedDataCachePath)
-Cloud=(Type=Cloud, ServerID=Cloud, EnvHostOverride=UE-CloudDataCacheHost, CommandLineHostOverride=CloudDataCacheHost, CommandLineOAuthSecretOverride="CloudDataCacheOAuthSecret", OAuthAccessTokenEnvOverride="UE-CloudDataCacheAccessToken", EnvHttpVersionOverride="UE-CloudDataCacheHttpVersion", CommandLineHttpVersionOverride="CloudDataCacheHttpVersion", EnvOAuthProviderIdentifierOverride="UE-CloudDataCacheOAuthProviderIdentifier", CommandLineOAuthProviderIdentifierOverride="CloudDataCacheOAuthProviderIdentifier")
+Cloud=(Type=Cloud, ServerID=Cloud, EnvHostOverride=UE-CloudDataCacheHost, CommandLineHostOverride=CloudDataCacheHost, CommandLineOAuthSecretOverride="CloudDataCacheOAuthSecret", EnvOAuthAccessTokenOverride="UE-CloudDataCacheAccessToken", EnvHttpVersionOverride="UE-CloudDataCacheHttpVersion", CommandLineHttpVersionOverride="CloudDataCacheHttpVersion", EnvOAuthProviderIdentifierOverride="UE-CloudDataCacheOAuthProviderIdentifier", CommandLineOAuthProviderIdentifierOverride="CloudDataCacheOAuthProviderIdentifier")
 Pak=(Type=ReadPak, Filename="%GAMEDIR%DerivedDataCache/DDC.ddp")
 CompressedPak=(Type=ReadPak, Filename="%GAMEDIR%DerivedDataCache/Compressed.ddp", Compressed=true)
 EnginePak=(Type=ReadPak, Filename=../../../Engine/DerivedDataCache/Compressed.ddp, Compressed=true)
@@ -2798,7 +2798,7 @@ Local=(Type=FileSystem, DeleteOnly=true, UnusedFileAge=8, Path=%ENGINEDIR%Derive
 ZenLocal=(Type=Zen, ServerID=Local)
 ZenShared=(Type=Zen, ServerID=Shared)
 Shared=(Type=FileSystem, UnusedFileAge=10, FoldersToClean=10, ConsiderSlowAt=70, Path=?EpicDDC, EnvPathOverride=UE-SharedDataCachePath, EditorOverrideSetting=SharedDerivedDataCache, CommandLineOverride=SharedDataCachePath)
-Cloud=(Type=Cloud, ServerID=Cloud, EnvHostOverride=UE-CloudDataCacheHost, CommandLineHostOverride=CloudDataCacheHost, CommandLineOAuthSecretOverride="CloudDataCacheOAuthSecret", OAuthAccessTokenEnvOverride="UE-CloudDataCacheAccessToken", EnvHttpVersionOverride="UE-CloudDataCacheHttpVersion", CommandLineHttpVersionOverride="CloudDataCacheHttpVersion", EnvOAuthProviderIdentifierOverride="UE-CloudDataCacheOAuthProviderIdentifier", CommandLineOAuthProviderIdentifierOverride="CloudDataCacheOAuthProviderIdentifier")
+Cloud=(Type=Cloud, ServerID=Cloud, EnvHostOverride=UE-CloudDataCacheHost, CommandLineHostOverride=CloudDataCacheHost, CommandLineOAuthSecretOverride="CloudDataCacheOAuthSecret", EnvOAuthAccessTokenOverride="UE-CloudDataCacheAccessToken", EnvHttpVersionOverride="UE-CloudDataCacheHttpVersion", CommandLineHttpVersionOverride="CloudDataCacheHttpVersion", EnvOAuthProviderIdentifierOverride="UE-CloudDataCacheOAuthProviderIdentifier", CommandLineOAuthProviderIdentifierOverride="CloudDataCacheOAuthProviderIdentifier")
 Pak=(Type=ReadPak, Filename="%GAMEDIR%DerivedDataCache/DDC.ddp")
 EnginePak=(Type=ReadPak, Filename="%ENGINEDIR%DerivedDataCache/DDC.ddp")
 
@@ -2852,7 +2852,7 @@ AsyncPut=(Type=AsyncPut, Inner=Hierarchy)
 Hierarchy=(Type=Hierarchical, Inner=PakWrite, Inner=PakRead, Inner=ZenLocal, Inner=Local, Inner=ZenShared, Inner=Shared, Inner=Cloud)
 Local=(Type=FileSystem, DeleteOnly=true, UnusedFileAge=8, Path=%ENGINEDIR%DerivedDataCache, EnvPathOverride=UE-LocalDataCachePath, EditorOverrideSetting=LocalDerivedDataCache, CommandLineOverride=LocalDataCachePath)
 Shared=(Type=FileSystem, DeleteUnused=false, ConsiderSlowAt=70, Path=?EpicDDC, EnvPathOverride=UE-SharedDataCachePath, EditorOverrideSetting=SharedDerivedDataCache, CommandLineOverride=SharedDataCachePath)
-Cloud=(Type=Cloud, ServerID=Cloud, EnvHostOverride=UE-CloudDataCacheHost, CommandLineHostOverride=CloudDataCacheHost, CommandLineOAuthSecretOverride="CloudDataCacheOAuthSecret", OAuthAccessTokenEnvOverride="UE-CloudDataCacheAccessToken", EnvHttpVersionOverride="UE-CloudDataCacheHttpVersion", CommandLineHttpVersionOverride="CloudDataCacheHttpVersion", EnvOAuthProviderIdentifierOverride="UE-CloudDataCacheOAuthProviderIdentifier", CommandLineOAuthProviderIdentifierOverride="CloudDataCacheOAuthProviderIdentifier")
+Cloud=(Type=Cloud, ServerID=Cloud, EnvHostOverride=UE-CloudDataCacheHost, CommandLineHostOverride=CloudDataCacheHost, CommandLineOAuthSecretOverride="CloudDataCacheOAuthSecret", EnvOAuthAccessTokenOverride="UE-CloudDataCacheAccessToken", EnvHttpVersionOverride="UE-CloudDataCacheHttpVersion", CommandLineHttpVersionOverride="CloudDataCacheHttpVersion", EnvOAuthProviderIdentifierOverride="UE-CloudDataCacheOAuthProviderIdentifier", CommandLineOAuthProviderIdentifierOverride="CloudDataCacheOAuthProviderIdentifier")
 PakRead=(Type=ReadPak, Filename="%GAMEDIR%DerivedDataCache/DDC.ddp")
 PakWrite=(Type=WritePak, Filename="%GAMEDIR%DerivedDataCache/DDC.ddp")
 ZenLocal=(Type=Zen, ServerID=Local)
@@ -2864,7 +2864,7 @@ AsyncPut=(Type=AsyncPut, Inner=Hierarchy)
 Hierarchy=(Type=Hierarchical, Inner=EnginePak, Inner=PakWrite, Inner=PakRead, Inner=ZenLocal, Inner=Local, Inner=ZenShared, Inner=Shared, Inner=Cloud)
 Local=(Type=FileSystem, DeleteOnly=true, UnusedFileAge=8, Path=%ENGINEDIR%DerivedDataCache, EnvPathOverride=UE-LocalDataCachePath, EditorOverrideSetting=LocalDerivedDataCache, CommandLineOverride=LocalDataCachePath)
 Shared=(Type=FileSystem, DeleteUnused=false, ConsiderSlowAt=70, Path=?EpicDDC, EnvPathOverride=UE-SharedDataCachePath, EditorOverrideSetting=SharedDerivedDataCache, CommandLineOverride=SharedDataCachePath)
-Cloud=(Type=Cloud, ServerID=Cloud, EnvHostOverride=UE-CloudDataCacheHost, CommandLineHostOverride=CloudDataCacheHost, CommandLineOAuthSecretOverride="CloudDataCacheOAuthSecret", OAuthAccessTokenEnvOverride="UE-CloudDataCacheAccessToken", EnvHttpVersionOverride="UE-CloudDataCacheHttpVersion", CommandLineHttpVersionOverride="CloudDataCacheHttpVersion", EnvOAuthProviderIdentifierOverride="UE-CloudDataCacheOAuthProviderIdentifier", CommandLineOAuthProviderIdentifierOverride="CloudDataCacheOAuthProviderIdentifier")
+Cloud=(Type=Cloud, ServerID=Cloud, EnvHostOverride=UE-CloudDataCacheHost, CommandLineHostOverride=CloudDataCacheHost, CommandLineOAuthSecretOverride="CloudDataCacheOAuthSecret", EnvOAuthAccessTokenOverride="UE-CloudDataCacheAccessToken", EnvHttpVersionOverride="UE-CloudDataCacheHttpVersion", CommandLineHttpVersionOverride="CloudDataCacheHttpVersion", EnvOAuthProviderIdentifierOverride="UE-CloudDataCacheOAuthProviderIdentifier", CommandLineOAuthProviderIdentifierOverride="CloudDataCacheOAuthProviderIdentifier")
 EnginePak=(Type=ReadPak, Filename=../../../Engine/DerivedDataCache/Compressed.ddp, Compressed=true)
 PakRead=(Type=ReadPak, Filename="%GAMEDIR%DerivedDataCache/Compressed.ddp", Compressed=true)
 PakWrite=(Type=WritePak, Filename="%GAMEDIR%DerivedDataCache/Compressed.ddp", Compressed=true)
@@ -2877,7 +2877,7 @@ AsyncPut=(Type=AsyncPut, Inner=Hierarchy)
 Hierarchy=(Type=Hierarchical, Inner=PakWrite, Inner=ZenLocal, Inner=Local, Inner=ZenShared, Inner=Shared, Inner=Cloud)
 Local=(Type=FileSystem, DeleteOnly=true, UnusedFileAge=8, Path=%ENGINEDIR%DerivedDataCache, EnvPathOverride=UE-LocalDataCachePath, EditorOverrideSetting=LocalDerivedDataCache, CommandLineOverride=LocalDataCachePath)
 Shared=(Type=FileSystem, DeleteUnused=false, ConsiderSlowAt=70, Path=?EpicDDC, EnvPathOverride=UE-SharedDataCachePath, EditorOverrideSetting=SharedDerivedDataCache, CommandLineOverride=SharedDataCachePath)
-Cloud=(Type=Cloud, ServerID=Cloud, EnvHostOverride=UE-CloudDataCacheHost, CommandLineHostOverride=CloudDataCacheHost, CommandLineOAuthSecretOverride="CloudDataCacheOAuthSecret", OAuthAccessTokenEnvOverride="UE-CloudDataCacheAccessToken", EnvHttpVersionOverride="UE-CloudDataCacheHttpVersion", CommandLineHttpVersionOverride="CloudDataCacheHttpVersion", EnvOAuthProviderIdentifierOverride="UE-CloudDataCacheOAuthProviderIdentifier", CommandLineOAuthProviderIdentifierOverride="CloudDataCacheOAuthProviderIdentifier")
+Cloud=(Type=Cloud, ServerID=Cloud, EnvHostOverride=UE-CloudDataCacheHost, CommandLineHostOverride=CloudDataCacheHost, CommandLineOAuthSecretOverride="CloudDataCacheOAuthSecret", EnvOAuthAccessTokenOverride="UE-CloudDataCacheAccessToken", EnvHttpVersionOverride="UE-CloudDataCacheHttpVersion", CommandLineHttpVersionOverride="CloudDataCacheHttpVersion", EnvOAuthProviderIdentifierOverride="UE-CloudDataCacheOAuthProviderIdentifier", CommandLineOAuthProviderIdentifierOverride="CloudDataCacheOAuthProviderIdentifier")
 PakWrite=(Type=WritePak, Filename="%ENGINEDIR%DerivedDataCache/Compressed.ddp", Compressed=true)
 ZenLocal=(Type=Zen, ServerID=Local)
 ZenShared=(Type=Zen, ServerID=Shared)
@@ -2888,7 +2888,7 @@ AsyncPut=(Type=AsyncPut, Inner=Hierarchy)
 Hierarchy=(Type=Hierarchical, Inner=PakWrite, Inner=ZenLocal, Inner=Local, Inner=ZenShared, Inner=Shared, Inner=Cloud)
 Local=(Type=FileSystem, DeleteOnly=true, UnusedFileAge=8, Path=%ENGINEDIR%DerivedDataCache, EnvPathOverride=UE-LocalDataCachePath, EditorOverrideSetting=LocalDerivedDataCache, CommandLineOverride=LocalDataCachePath)
 Shared=(Type=FileSystem, DeleteUnused=false, ConsiderSlowAt=70, Path=?EpicDDC, EnvPathOverride=UE-SharedDataCachePath, EditorOverrideSetting=SharedDerivedDataCache, CommandLineOverride=SharedDataCachePath)
-Cloud=(Type=Cloud, ServerID=Cloud, EnvHostOverride=UE-CloudDataCacheHost, CommandLineHostOverride=CloudDataCacheHost, CommandLineOAuthSecretOverride="CloudDataCacheOAuthSecret", OAuthAccessTokenEnvOverride="UE-CloudDataCacheAccessToken", EnvHttpVersionOverride="UE-CloudDataCacheHttpVersion", CommandLineHttpVersionOverride="CloudDataCacheHttpVersion", EnvOAuthProviderIdentifierOverride="UE-CloudDataCacheOAuthProviderIdentifier", CommandLineOAuthProviderIdentifierOverride="CloudDataCacheOAuthProviderIdentifier")
+Cloud=(Type=Cloud, ServerID=Cloud, EnvHostOverride=UE-CloudDataCacheHost, CommandLineHostOverride=CloudDataCacheHost, CommandLineOAuthSecretOverride="CloudDataCacheOAuthSecret", EnvOAuthAccessTokenOverride="UE-CloudDataCacheAccessToken", EnvHttpVersionOverride="UE-CloudDataCacheHttpVersion", CommandLineHttpVersionOverride="CloudDataCacheHttpVersion", EnvOAuthProviderIdentifierOverride="UE-CloudDataCacheOAuthProviderIdentifier", CommandLineOAuthProviderIdentifierOverride="CloudDataCacheOAuthProviderIdentifier")
 PakWrite=(Type=WritePak, Filename=../../../Enterprise/DerivedDataCache/Compressed.ddp, Compressed=true)
 ZenLocal=(Type=Zen, ServerID=Local)
 ZenShared=(Type=Zen, ServerID=Shared)
@@ -2900,7 +2900,7 @@ Hierarchy=(Type=Hierarchical, Inner=ZenLocal, Inner=Local, Inner=Project, Inner=
 Local=(Type=FileSystem, DeleteOnly=true, UnusedFileAge=8, Path=%ENGINEDIR%DerivedDataCache, EnvPathOverride=UE-LocalDataCachePath, EditorOverrideSetting=LocalDerivedDataCache, CommandLineOverride=LocalDataCachePath)
 Project=(Type=FileSystem, Flush=true, UnusedFileAge=34, Path=%GAMEDIR%ProjectDerivedData)
 Shared=(Type=FileSystem, DeleteUnused=false, ConsiderSlowAt=70, Path=?EpicDDC, EnvPathOverride=UE-SharedDataCachePath, EditorOverrideSetting=SharedDerivedDataCache, CommandLineOverride=SharedDataCachePath)
-Cloud=(Type=Cloud, ServerID=Cloud, EnvHostOverride=UE-CloudDataCacheHost, CommandLineHostOverride=CloudDataCacheHost, CommandLineOAuthSecretOverride="CloudDataCacheOAuthSecret", OAuthAccessTokenEnvOverride="UE-CloudDataCacheAccessToken")
+Cloud=(Type=Cloud, ServerID=Cloud, EnvHostOverride=UE-CloudDataCacheHost, CommandLineHostOverride=CloudDataCacheHost, CommandLineOAuthSecretOverride="CloudDataCacheOAuthSecret", EnvOAuthAccessTokenOverride="UE-CloudDataCacheAccessToken")
 ZenLocal=(Type=Zen, ServerID=Local)
 ZenShared=(Type=Zen, ServerID=Shared)
 
diff --git a/Engine/Source/Developer/DerivedDataCache/Private/HttpCacheStore.cpp b/Engine/Source/Developer/DerivedDataCache/Private/HttpCacheStore.cpp
--- a/Engine/Source/Developer/DerivedDataCache/Private/HttpCacheStore.cpp
+++ b/Engine/Source/Developer/DerivedDataCache/Private/HttpCacheStore.cpp
@@ -239,6 +239,7 @@ struct FHttpCacheStoreParams
 
 	bool bResolveHostCanonicalName = true;
 	bool bReadOnly = false;
+	bool bWriteOnly = false;
 	bool bBypassProxy = false;
 
 	void Parse(const TCHAR* NodeName, const TCHAR* Config);
@@ -347,6 +348,7 @@ private:
 
 	bool bIsUsable = false;
 	bool bReadOnly = false;
+	bool bWriteOnly = false;
 	bool bBypassProxy = false;
 
 	static inline FHttpCacheStore* AnyInstance = nullptr;
@@ -2270,6 +2272,7 @@ FHttpCacheStore::FHttpCacheStore(const FHttpCacheStoreParams& Params, ICacheStor
 	, AuthScheme(Params.AuthScheme)
 	, StoreOwner(Owner)
 	, bReadOnly(Params.bReadOnly)
+	, bWriteOnly(Params.bWriteOnly)
 	, bBypassProxy(Params.bBypassProxy)
 {
 	TRACE_CPUPROFILER_EVENT_SCOPE(HttpDDC_Construct);
@@ -2361,8 +2364,9 @@ FHttpCacheStore::FHttpCacheStore(const FHttpCacheStoreParams& Params, ICacheStor
 
 		if (StoreOwner)
 		{
-			const ECacheStoreFlags Flags = ECacheStoreFlags::Remote | ECacheStoreFlags::Query |
-				(Params.bReadOnly ? ECacheStoreFlags::None : ECacheStoreFlags::Store);
+			const ECacheStoreFlags Flags = ECacheStoreFlags::Remote |
+											(Params.bReadOnly ? ECacheStoreFlags::None : ECacheStoreFlags::Store) |
+											(Params.bWriteOnly ? ECacheStoreFlags::None : ECacheStoreFlags::Query);
 			TStringBuilder<256> Path(InPlace, ResolvedHost, TEXTVIEW(" ("), Namespace, TEXTVIEW(")"));
 			StoreOwner->Add(this, Flags);
 			StoreStats = StoreOwner->CreateStats(this, Flags, TEXTVIEW("Unreal Cloud DDC"), Params.Name, Path);
@@ -3301,12 +3305,60 @@ void FHttpCacheStore::GetChunks(
 	GetChunkGroupAsync(Owner, PendingGroupStartRequest, SortedRequests.GetData() + SortedRequests.Num(), SharedOnComplete);
 }
 
-void FHttpCacheStoreParams::Parse(const TCHAR* NodeName, const TCHAR* Config)
+static void RegisterInheritedCommandLineArg(const FStringView ArgName)
+{
+	FCommandLine::RegisterArgument(ArgName, ECommandLineArgumentFlags::EditorContext | ECommandLineArgumentFlags::CommandletContext | ECommandLineArgumentFlags::Inherit);
+}
+
+template <typename T>
+static bool ParseConfigValue(const TCHAR* Stream, const TCHAR* Match, T& OutValue)
 {
-	auto RegisterInheritedCommandlineArg = [](const FStringView ArgName)
+	return FParse::Value(Stream, Match, OutValue);
+}
+
+template <>
+bool ParseConfigValue<bool>(const TCHAR* Stream, const TCHAR* Match, bool& OutValue)
+{
+	return FParse::Bool(Stream, Match, OutValue);
+}
+
+static bool ConvertStringToParamValue(const FString& InString, FString& OutString)
+{
+	OutString = InString;
+	return true;
+}
+
+static bool ConvertStringToParamValue(const FString& InString, bool& OutBool)
+{
+	OutBool = FCString::ToBool(*InString);
+	return true;
+}
+
+template <typename T>
+static void ParseConfigParamWithOverrides(const TCHAR* NodeName, const TCHAR* Config, const TCHAR* ParamName, T& OutValue, bool bIsSecret)
+{
+	FString OverrideName;
+	ParseConfigValue(Config, *WriteToString<32>(ParamName, TEXT("=")), OutValue);
+	if (ParseConfigValue(Config, *WriteToString<32>(TEXT("Env"), ParamName, TEXT("Override=")), OverrideName))
 	{
-		FCommandLine::RegisterArgument(ArgName, ECommandLineArgumentFlags::EditorContext | ECommandLineArgumentFlags::CommandletContext | ECommandLineArgumentFlags::Inherit);
-	};
+		FString ParamEnv = FPlatformMisc::GetEnvironmentVariable(*OverrideName);
+		if (!ParamEnv.IsEmpty() && ConvertStringToParamValue(ParamEnv, OutValue))
+		{
+			UE_LOG(LogDerivedDataCache, Log, TEXT("%s: Found environment override for %s %s=%s"), NodeName, ParamName, *OverrideName, bIsSecret ? TEXT("{SECRET}") : *ParamEnv);
+		}
+	}
+	if (ParseConfigValue(Config, *WriteToString<32>(TEXT("CommandLine"), ParamName, TEXT("Override=")), OverrideName))
+	{
+		if (ParseConfigValue(FCommandLine::Get(), *WriteToString<32>(OverrideName, TEXT("=")), OutValue))
+		{
+			RegisterInheritedCommandLineArg(OverrideName);
+			UE_LOG(LogDerivedDataCache, Log, TEXT("%s: Found command line override for %s %s=%s"), NodeName, ParamName, *OverrideName, bIsSecret ? TEXT("{SECRET}") : *WriteToString<32>(OutValue));
+		}
+	}
+}
+
+void FHttpCacheStoreParams::Parse(const TCHAR* NodeName, const TCHAR* Config)
+{
 	Name = NodeName;
 
 	FString ServerId;
@@ -3332,25 +3384,7 @@ void FHttpCacheStoreParams::Parse(const TCHAR* NodeName, const TCHAR* Config)
 	FString OverrideName;
 
 	// Host Params
-
-	FParse::Value(Config, TEXT("Host="), Host);
-	if (FParse::Value(Config, TEXT("EnvHostOverride="), OverrideName))
-	{
-		FString HostEnv = FPlatformMisc::GetEnvironmentVariable(*OverrideName);
-		if (!HostEnv.IsEmpty())
-		{
-			Host = HostEnv;
-			UE_LOG(LogDerivedDataCache, Log, TEXT("%s: Found environment override for Host %s=%s"), NodeName, *OverrideName, *Host);
-		}
-	}
-	if (FParse::Value(Config, TEXT("CommandLineHostOverride="), OverrideName))
-	{
-		if (FParse::Value(FCommandLine::Get(), *(OverrideName + TEXT("=")), Host))
-		{
-			RegisterInheritedCommandlineArg(OverrideName);
-			UE_LOG(LogDerivedDataCache, Log, TEXT("%s: Found command line override for Host %s=%s"), NodeName, *OverrideName, *Host);
-		}
-	}
+	ParseConfigParamWithOverrides(NodeName, Config, TEXT("Host"), Host, /*bIsSecret*/false);
 
 	FParse::Value(Config, TEXT("DiscoveryHost="), DiscoveryHost);
 
@@ -3359,45 +3393,10 @@ void FHttpCacheStoreParams::Parse(const TCHAR* NodeName, const TCHAR* Config)
 	FParse::Bool(Config, TEXT("ResolveHostCanonicalName="), bResolveHostCanonicalName);
 
 	// Http version Params
-
-	FParse::Value(Config, TEXT("HttpVersion="), HttpVersion);
-	if (FParse::Value(Config, TEXT("EnvHttpVersionOverride="), OverrideName))
-	{
-		FString HttpEnv = FPlatformMisc::GetEnvironmentVariable(*OverrideName);
-		if (!HttpEnv.IsEmpty())
-		{
-			HttpVersion = HttpEnv;
-			UE_LOG(LogDerivedDataCache, Log, TEXT("%s: Found environment override for HttpVersion %s=%s"), NodeName, *OverrideName, *HttpVersion);
-		}
-	}
-	if (FParse::Value(Config, TEXT("CommandLineHttpVersionOverride="), OverrideName))
-	{
-		if (FParse::Value(FCommandLine::Get(), *(OverrideName + TEXT("=")), HttpVersion))
-		{
-			RegisterInheritedCommandlineArg(OverrideName);
-			UE_LOG(LogDerivedDataCache, Log, TEXT("%s: Found command line override for HttpVersion %s=%s"), NodeName, *OverrideName, *HttpVersion);
-		}
-	}
+	ParseConfigParamWithOverrides(NodeName, Config, TEXT("HttpVersion"), HttpVersion, /*bIsSecret*/false);
 
 	// Unix Socket Params
-	FParse::Value(Config, TEXT("UnixSocketPath="), UnixSocketPath);
-	if (FParse::Value(Config, TEXT("EnvUnixSocketPathOverride="), OverrideName))
-	{
-		FString UnixSocketPathEnv = FPlatformMisc::GetEnvironmentVariable(*OverrideName);
-		if (!UnixSocketPathEnv.IsEmpty())
-		{
-			UnixSocketPath = UnixSocketPathEnv;
-			UE_LOG(LogDerivedDataCache, Log, TEXT("%s: Found environment override for UnixSocketPath %s=%s"), NodeName, *OverrideName, *UnixSocketPath);
-		}
-	}
-	if (FParse::Value(Config, TEXT("CommandLineUnixSocketPathOverride="), OverrideName))
-	{
-		if (FParse::Value(FCommandLine::Get(), *(OverrideName + TEXT("=")), UnixSocketPath))
-		{
-			RegisterInheritedCommandlineArg(OverrideName);
-			UE_LOG(LogDerivedDataCache, Log, TEXT("%s: Found command line override for UnixSocketPath %s=%s"), NodeName, *OverrideName, *UnixSocketPath);
-		}
-	}
+	ParseConfigParamWithOverrides(NodeName, Config, TEXT("UnixSocketPath"), UnixSocketPath, /*bIsSecret*/false);
 
 	// Namespace Params
 
@@ -3408,38 +3407,11 @@ void FHttpCacheStoreParams::Parse(const TCHAR* NodeName, const TCHAR* Config)
 	FParse::Value(Config, TEXT("StructuredNamespace="), Namespace);
 
 	// OAuth Params
-
-	FParse::Value(Config, TEXT("OAuthProvider="), OAuthProvider);
-
-	if (FParse::Value(Config, TEXT("CommandLineOAuthProviderOverride="), OverrideName))
-	{
-		if (FParse::Value(FCommandLine::Get(), *(OverrideName + TEXT("=")), OAuthProvider))
-		{
-			RegisterInheritedCommandlineArg(OverrideName);
-			UE_LOG(LogDerivedDataCache, Log, TEXT("%s: Found command line override for OAuthProvider %s=%s"), NodeName, *OverrideName, *OAuthProvider);
-		}
-	}
+	ParseConfigParamWithOverrides(NodeName, Config, TEXT("OAuthProvider"), OAuthProvider, /*bIsSecret*/false);
 
 	FParse::Value(Config, TEXT("OAuthClientId="), OAuthClientId);
 
-	FParse::Value(Config, TEXT("OAuthSecret="), OAuthSecret);
-	if (FParse::Value(Config, TEXT("EnvOAuthSecretOverride="), OverrideName))
-	{
-		FString OAuthSecretEnv = FPlatformMisc::GetEnvironmentVariable(*OverrideName);
-		if (!OAuthSecretEnv.IsEmpty())
-		{
-			OAuthSecret = OAuthSecretEnv;
-			UE_LOG(LogDerivedDataCache, Log, TEXT("%s: Found environment override for OAuthSecret %s={SECRET}"), NodeName, *OverrideName);
-		}
-	}
-	if (FParse::Value(Config, TEXT("CommandLineOAuthSecretOverride="), OverrideName))
-	{
-		if (FParse::Value(FCommandLine::Get(), *(OverrideName + TEXT("=")), OAuthSecret))
-		{
-			RegisterInheritedCommandlineArg(OverrideName);
-			UE_LOG(LogDerivedDataCache, Log, TEXT("%s: Found command line override for OAuthSecret %s={SECRET}"), NodeName, *OverrideName);
-		}
-	}
+	ParseConfigParamWithOverrides(NodeName, Config, TEXT("OAuthSecret"), OAuthSecret, /*bIsSecret*/true);
 
 	// If the secret is a file path, read the secret from the file.
 	if (OAuthSecret.StartsWith(TEXT("file://")))
@@ -3455,28 +3427,32 @@ void FHttpCacheStoreParams::Parse(const TCHAR* NodeName, const TCHAR* Config)
 
 	FParse::Value(Config, TEXT("OAuthScope="), OAuthScope);
 
-    // OAuth Provider Identifier
-	FParse::Value(Config, TEXT("OAuthProviderIdentifier="), OAuthProviderIdentifier);
-	if (FParse::Value(Config, TEXT("EnvOAuthProviderIdentifierOverride="), OverrideName))
+	// OAuth Provider Identifier
+	ParseConfigParamWithOverrides(NodeName, Config, TEXT("OAuthProviderIdentifier"), OAuthProviderIdentifier, /*bIsSecret*/false);
+
+	// OAuth Access
+	// TODO-BEGIN: Replace this with:
+	// ParseConfigParamWithOverrides(NodeName, Config, TEXT("OAuthAccessToken"), OAuthAccessToken, /*bIsSecret*/true);
+	if (!FParse::Value(Config, TEXT("OAuthAccessToken="), OAuthAccessToken))
 	{
-		FString ProviderEnv = FPlatformMisc::GetEnvironmentVariable(*OverrideName);
-		if (!ProviderEnv.IsEmpty())
+		if (FParse::Value(Config, TEXT("OAuthAccess="), OAuthAccessToken))
 		{
-			OAuthProviderIdentifier = ProviderEnv;
-			UE_LOG(LogDerivedDataCache, Log, TEXT("%s: Found environment override for OAuthProviderIdentifier %s=%s"), NodeName, *OverrideName, *OAuthProviderIdentifier);
+			UE_LOG(LogDerivedDataCache, Warning, TEXT("%s: Deprecated parameter 'OAuthAccess' found.  Please use 'OAuthAccessToken' instead."), NodeName);
 		}
 	}
-	if (FParse::Value(Config, TEXT("CommandLineOAuthProviderIdentifierOverride="), OverrideName))
+	if (FParse::Value(Config, TEXT("EnvOAuthAccessTokenOverride="), OverrideName))
 	{
-		if (FParse::Value(FCommandLine::Get(), *(OverrideName + TEXT("=")), OAuthProviderIdentifier))
+		FString AccessToken = FPlatformMisc::GetEnvironmentVariable(*OverrideName);
+		if (!AccessToken.IsEmpty())
 		{
-			RegisterInheritedCommandlineArg(OverrideName);
-			UE_LOG(LogDerivedDataCache, Log, TEXT("%s: Found command line override for OAuthProviderIdentifier %s=%s"), NodeName, *OverrideName, *OAuthProviderIdentifier);
+			OAuthAccessToken = AccessToken;
+			// We do not log the access token as it is sensitive information.
+			UE_LOG(LogDerivedDataCache, Log, TEXT("%s: Found environment override for OAuthAccess %s={SECRET}"), NodeName, *OverrideName);
 		}
 	}
-	FParse::Value(Config, TEXT("OAuthAccess="), OAuthAccessToken);
-	if (FParse::Value(Config, TEXT("OAuthAccessTokenEnvOverride="), OverrideName))
+	else if (FParse::Value(Config, TEXT("OAuthAccessTokenEnvOverride="), OverrideName))
 	{
+		UE_LOG(LogDerivedDataCache, Warning, TEXT("%s: Deprecated parameter 'OAuthAccessTokenEnvOverride' found.  Please use 'EnvOAuthAccessTokenOverride' instead."), NodeName);
 		FString AccessToken = FPlatformMisc::GetEnvironmentVariable(*OverrideName);
 		if (!AccessToken.IsEmpty())
 		{
@@ -3485,6 +3461,15 @@ void FHttpCacheStoreParams::Parse(const TCHAR* NodeName, const TCHAR* Config)
 			UE_LOG(LogDerivedDataCache, Log, TEXT("%s: Found OAuth access token in %s."), NodeName, *OverrideName);
 		}
 	}
+	if (ParseConfigValue(Config, TEXT("CommandLineOAuthAccessTokenOverride="), OverrideName))
+	{
+		if (ParseConfigValue(FCommandLine::Get(), *WriteToString<32>(OverrideName, TEXT("=")), OAuthAccessToken))
+		{
+			RegisterInheritedCommandLineArg(OverrideName);
+			UE_LOG(LogDerivedDataCache, Log, TEXT("%s: Found command line override for OAuthAccess %s={SECRET}"), NodeName, *OverrideName);
+		}
+	}
+	// TODO-END
 
 	FParse::Value(Config, TEXT("AuthScheme="), AuthScheme);
 	if (AuthScheme.IsEmpty())
@@ -3495,8 +3480,14 @@ void FHttpCacheStoreParams::Parse(const TCHAR* NodeName, const TCHAR* Config)
 	FParse::Value(Config, TEXT("OAuthPinnedPublicKeys="), OAuthPinnedPublicKeys);
 
 	// Cache Params
+	ParseConfigParamWithOverrides(NodeName, Config, TEXT("ReadOnly"), bReadOnly, /*bIsSecret*/false);
+	ParseConfigParamWithOverrides(NodeName, Config, TEXT("WriteOnly"), bWriteOnly, /*bIsSecret*/false);
+
+	if (bReadOnly && bWriteOnly)
+	{
+		UE_LOG(LogDerivedDataCache, Display, TEXT("%s: Both ReadOnly and WriteOnly specified.  This cache store will be inactive."), NodeName);
+	}
 
-	FParse::Bool(Config, TEXT("ReadOnly="), bReadOnly);
 	FParse::Bool(Config, TEXT("BypassProxy="), bBypassProxy);
 }
 
