Fill out Low/Medium/High scalability levels for r.Shadow.Virtual.MarkCoarsePagesLocal

diff --git a/Engine/Config/BaseScalability.ini b/Engine/Config/BaseScalability.ini
--- a/Engine/Config/BaseScalability.ini
+++ b/Engine/Config/BaseScalability.ini
@@ -152,6 +152,7 @@ r.Shadow.Virtual.SMRT.RayCountDirectional=0
 r.Shadow.Virtual.SMRT.SamplesPerRayDirectional=2
 r.Shadow.Virtual.SMRT.RayCountLocal=0
 r.Shadow.Virtual.SMRT.SamplesPerRayLocal=2
+r.Shadow.Virtual.MarkCoarsePagesLocal=2
 
 [ShadowQuality@1]
 r.LightFunctionQuality=1
@@ -176,6 +177,7 @@ r.Shadow.Virtual.SMRT.RayCountDirectional=4
 r.Shadow.Virtual.SMRT.SamplesPerRayDirectional=2
 r.Shadow.Virtual.SMRT.RayCountLocal=4
 r.Shadow.Virtual.SMRT.SamplesPerRayLocal=2
+r.Shadow.Virtual.MarkCoarsePagesLocal=2
 
 [ShadowQuality@2]
 r.LightFunctionQuality=1
@@ -203,6 +205,7 @@ r.Shadow.Virtual.SMRT.RayCountDirectional=8
 r.Shadow.Virtual.SMRT.SamplesPerRayDirectional=4
 r.Shadow.Virtual.SMRT.RayCountLocal=4
 r.Shadow.Virtual.SMRT.SamplesPerRayLocal=4
+r.Shadow.Virtual.MarkCoarsePagesLocal=2
 
 [ShadowQuality@3]
 r.LightFunctionQuality=1
