DDC: Removed AsyncPut, KeyLength, Verify nodes from graph configurations because they will be deprecated

diff --git a/Engine/Config/BaseEngine.ini b/Engine/Config/BaseEngine.ini
--- a/Engine/Config/BaseEngine.ini
+++ b/Engine/Config/BaseEngine.ini
@@ -2715,7 +2715,7 @@ Shared=(Host=None, Namespace="ue.ddc", EnvHostOverride=UE-ZenSharedDataCacheHost
 ; Others can be specified on the command line using:
 ;     -DDC=GraphSectionName     (for example: -DDC=NoShared)
 ; Each graph should start with a 'Root' node. The names of the other nodes are not predefined.
-; Supported node types are: KeyLength, AsyncPut, Hierarchical, FileSystem, Zen, Cloud, S3, ReadPak, WritePak, Verify
+; Supported node types are: Hierarchical, FileSystem, Zen, Cloud, S3, ReadPak, WritePak
 ; The order that nodes are defined in is not relevant.
 ; FileSystem backends can be disabled by setting the corresponding environment variable to 'None'. E.g., UE-SharedDataCachePath=None
 ; Cloud backends can be disabled by setting the corresponding environment variable to 'None'. E.g., UE-CloudDataCachePath=None
@@ -2756,9 +2756,7 @@ ZenLocal=(Type=Zen, ServerID=Local)
 ZenShared=(Type=Zen, ServerID=Shared)
 
 [InstalledDerivedDataBackendGraph]
-Root=(Type=KeyLength, Length=120, Inner=AsyncPut)
-AsyncPut=(Type=AsyncPut, Inner=Hierarchy)
-Hierarchy=(Type=Hierarchical, Inner=Pak, Inner=CompressedPak, Inner=EnginePak, Inner=EnterprisePak, Inner=ZenLocal, Inner=Local, Inner=ZenShared, Inner=Shared, Inner=Cloud)
+Root=(Type=Hierarchical, Inner=Pak, Inner=CompressedPak, Inner=EnginePak, Inner=EnterprisePak, Inner=ZenLocal, Inner=Local, Inner=ZenShared, Inner=Shared, Inner=Cloud)
 Local=(Type=FileSystem, DeleteOnly=true, UnusedFileAge=8, Path="%ENGINEVERSIONAGNOSTICUSERDIR%DerivedDataCache", EnvPathOverride=UE-LocalDataCachePath, EditorOverrideSetting=LocalDerivedDataCache)
 Shared=(Type=FileSystem, UnusedFileAge=10, FoldersToClean=10, ConsiderSlowAt=70, Path=?EpicDDC, EnvPathOverride=UE-SharedDataCachePath, EditorOverrideSetting=SharedDerivedDataCache, CommandLineOverride=SharedDataCachePath)
 Cloud=(Type=Cloud, ServerID=Cloud, EnvHostOverride=UE-CloudDataCacheHost, CommandLineHostOverride=CloudDataCacheHost, CommandLineOAuthSecretOverride="CloudDataCacheOAuthSecret", EnvOAuthAccessTokenOverride="UE-CloudDataCacheAccessToken", EnvHttpVersionOverride="UE-CloudDataCacheHttpVersion", CommandLineHttpVersionOverride="CloudDataCacheHttpVersion", EnvOAuthProviderIdentifierOverride="UE-CloudDataCacheOAuthProviderIdentifier", CommandLineOAuthProviderIdentifierOverride="CloudDataCacheOAuthProviderIdentifier")
@@ -2790,8 +2788,7 @@ EnterprisePak=(Type=ReadPak, Filename=../../../Enterprise/DerivedDataCache/Compr
 ZenShared=(Type=Zen, ServerID=Shared)
 
 [ZenOnlyDDC]
-Root=(Type=KeyLength, Length=120, Inner=AsyncPut)
-AsyncPut=(Type=AsyncPut, Inner=Zen)
+Root=(Type=Hierarchical, Inner=Zen)
 Zen=(Type=Zen, ServerID=Local)
 
 [ZenDDC]
@@ -2805,21 +2802,15 @@ Pak=(Type=ReadPak, Filename="%GAMEDIR%DerivedDataCache/DDC.ddp")
 EnginePak=(Type=ReadPak, Filename="%ENGINEDIR%DerivedDataCache/DDC.ddp")
 
 [ZenColdDDC]
-Root=(Type=KeyLength, Length=120, Inner=AsyncPut)
-AsyncPut=(Type=AsyncPut, Inner=Hierarchy)
-Hierarchy=(Type=Hierarchical, Inner=ZenLocal)
+Root=(Type=Hierarchical, Inner=ZenLocal)
 ZenLocal=(Type=Zen, ServerID=Local, Sandbox=Cold, Flush=true)
 
 [ZenWarmDDC]
-Root=(Type=KeyLength, Length=120, Inner=AsyncPut)
-AsyncPut=(Type=AsyncPut, Inner=Hierarchy)
-Hierarchy=(Type=Hierarchical, Inner=ZenLocal)
+Root=(Type=Hierarchical, Inner=ZenLocal)
 ZenLocal=(Type=Zen, ServerID=Local, Sandbox=Cold)
 
 [DerivedDataBackendGraph_Fill_Seattle]
-Root=(Type=KeyLength, Length=120, Inner=AsyncPut)
-AsyncPut=(Type=AsyncPut, Inner=Hierarchy)
-Hierarchy=(Type=Hierarchical, Inner=Pak, Inner=EnginePak, Inner=ZenLocal, Inner=Local, Inner=Seattle)
+Root=(Type=Hierarchical, Inner=Pak, Inner=EnginePak, Inner=ZenLocal, Inner=Local, Inner=Seattle)
 Local=(Type=FileSystem, DeleteOnly=true, UnusedFileAge=8, Path=%ENGINEDIR%DerivedDataCache, EnvPathOverride=UE-LocalDataCachePath, EditorOverrideSetting=LocalDerivedDataCache, CommandLineOverride=LocalDataCachePath)
 Seattle=(Type=FileSystem, UnusedFileAge=23, FoldersToClean=10, Path=?EpicSeaDDC, EnvPathOverride=UE-SharedDataCachePath_Seattle)
 Pak=(Type=ReadPak, Filename="%GAMEDIR%DerivedDataCache/DDC.ddp")
@@ -2827,31 +2818,23 @@ EnginePak=(Type=ReadPak, Filename="%ENGINEDIR%DerivedDataCache/DDC.ddp")
 ZenLocal=(Type=Zen, ServerID=Local)
 
 [NoShared]
-Root=(Type=KeyLength, Length=120, Inner=AsyncPut)
-AsyncPut=(Type=AsyncPut, Inner=Hierarchy)
-Hierarchy=(Type=Hierarchical, Inner=Pak, Inner=ZenLocal, Inner=Local)
+Root=(Type=Hierarchical, Inner=Pak, Inner=ZenLocal, Inner=Local)
 Local=(Type=FileSystem, DeleteOnly=true, UnusedFileAge=8, Path=%ENGINEDIR%DerivedDataCache, EnvPathOverride=UE-LocalDataCachePath, EditorOverrideSetting=LocalDerivedDataCache, CommandLineOverride=LocalDataCachePath)
 Pak=(Type=ReadPak, Filename="%GAMEDIR%DerivedDataCache/DDC.ddp")
 ZenLocal=(Type=Zen, ServerID=Local)
 
 ; Creates a completely new standalone cache so as to not destroy a valuable populated local cache
 [Cold]
-Root=(Type=KeyLength, Length=120, Inner=AsyncPut)
-AsyncPut=(Type=AsyncPut, Inner=Hierarchy)
-Hierarchy=(Type=Hierarchical, Inner=ZenLocal)
+Root=(Type=Hierarchical, Inner=ZenLocal)
 ZenLocal=(Type=Zen, ServerID=Local, Sandbox=Cold, Flush=true)
 
 ; Inherits the same standalone cache from [Cold], for incremental performance measurement purposes
 [Warm]
-Root=(Type=KeyLength, Length=120, Inner=AsyncPut)
-AsyncPut=(Type=AsyncPut, Inner=Hierarchy)
-Hierarchy=(Type=Hierarchical, Inner=ZenLocal)
+Root=(Type=Hierarchical, Inner=ZenLocal)
 ZenLocal=(Type=Zen, ServerID=Local, Sandbox=Cold)
 
 [CreatePak]
-Root=(Type=KeyLength, Length=120, Inner=AsyncPut)
-AsyncPut=(Type=AsyncPut, Inner=Hierarchy)
-Hierarchy=(Type=Hierarchical, Inner=PakWrite, Inner=PakRead, Inner=ZenLocal, Inner=Local, Inner=ZenShared, Inner=Shared, Inner=Cloud)
+Root=(Type=Hierarchical, Inner=PakWrite, Inner=PakRead, Inner=ZenLocal, Inner=Local, Inner=ZenShared, Inner=Shared, Inner=Cloud)
 Local=(Type=FileSystem, DeleteOnly=true, UnusedFileAge=8, Path=%ENGINEDIR%DerivedDataCache, EnvPathOverride=UE-LocalDataCachePath, EditorOverrideSetting=LocalDerivedDataCache, CommandLineOverride=LocalDataCachePath)
 Shared=(Type=FileSystem, DeleteUnused=false, ConsiderSlowAt=70, Path=?EpicDDC, EnvPathOverride=UE-SharedDataCachePath, EditorOverrideSetting=SharedDerivedDataCache, CommandLineOverride=SharedDataCachePath)
 Cloud=(Type=Cloud, ServerID=Cloud, EnvHostOverride=UE-CloudDataCacheHost, CommandLineHostOverride=CloudDataCacheHost, CommandLineOAuthSecretOverride="CloudDataCacheOAuthSecret", EnvOAuthAccessTokenOverride="UE-CloudDataCacheAccessToken", EnvHttpVersionOverride="UE-CloudDataCacheHttpVersion", CommandLineHttpVersionOverride="CloudDataCacheHttpVersion", EnvOAuthProviderIdentifierOverride="UE-CloudDataCacheOAuthProviderIdentifier", CommandLineOAuthProviderIdentifierOverride="CloudDataCacheOAuthProviderIdentifier")
@@ -2861,9 +2844,7 @@ ZenLocal=(Type=Zen, ServerID=Local)
 ZenShared=(Type=Zen, ServerID=Shared)
 
 [CreateInstalledProjectPak]
-Root=(Type=KeyLength, Length=120, Inner=AsyncPut)
-AsyncPut=(Type=AsyncPut, Inner=Hierarchy)
-Hierarchy=(Type=Hierarchical, Inner=EnginePak, Inner=PakWrite, Inner=PakRead, Inner=ZenLocal, Inner=Local, Inner=ZenShared, Inner=Shared, Inner=Cloud)
+Root=(Type=Hierarchical, Inner=EnginePak, Inner=PakWrite, Inner=PakRead, Inner=ZenLocal, Inner=Local, Inner=ZenShared, Inner=Shared, Inner=Cloud)
 Local=(Type=FileSystem, DeleteOnly=true, UnusedFileAge=8, Path=%ENGINEDIR%DerivedDataCache, EnvPathOverride=UE-LocalDataCachePath, EditorOverrideSetting=LocalDerivedDataCache, CommandLineOverride=LocalDataCachePath)
 Shared=(Type=FileSystem, DeleteUnused=false, ConsiderSlowAt=70, Path=?EpicDDC, EnvPathOverride=UE-SharedDataCachePath, EditorOverrideSetting=SharedDerivedDataCache, CommandLineOverride=SharedDataCachePath)
 Cloud=(Type=Cloud, ServerID=Cloud, EnvHostOverride=UE-CloudDataCacheHost, CommandLineHostOverride=CloudDataCacheHost, CommandLineOAuthSecretOverride="CloudDataCacheOAuthSecret", EnvOAuthAccessTokenOverride="UE-CloudDataCacheAccessToken", EnvHttpVersionOverride="UE-CloudDataCacheHttpVersion", CommandLineHttpVersionOverride="CloudDataCacheHttpVersion", EnvOAuthProviderIdentifierOverride="UE-CloudDataCacheOAuthProviderIdentifier", CommandLineOAuthProviderIdentifierOverride="CloudDataCacheOAuthProviderIdentifier")
@@ -2874,9 +2855,7 @@ ZenLocal=(Type=Zen, ServerID=Local)
 ZenShared=(Type=Zen, ServerID=Shared)
 
 [CreateInstalledEnginePak]
-Root=(Type=KeyLength, Length=120, Inner=AsyncPut)
-AsyncPut=(Type=AsyncPut, Inner=Hierarchy)
-Hierarchy=(Type=Hierarchical, Inner=PakWrite, Inner=ZenLocal, Inner=Local, Inner=ZenShared, Inner=Shared, Inner=Cloud)
+Root=(Type=Hierarchical, Inner=PakWrite, Inner=ZenLocal, Inner=Local, Inner=ZenShared, Inner=Shared, Inner=Cloud)
 Local=(Type=FileSystem, DeleteOnly=true, UnusedFileAge=8, Path=%ENGINEDIR%DerivedDataCache, EnvPathOverride=UE-LocalDataCachePath, EditorOverrideSetting=LocalDerivedDataCache, CommandLineOverride=LocalDataCachePath)
 Shared=(Type=FileSystem, DeleteUnused=false, ConsiderSlowAt=70, Path=?EpicDDC, EnvPathOverride=UE-SharedDataCachePath, EditorOverrideSetting=SharedDerivedDataCache, CommandLineOverride=SharedDataCachePath)
 Cloud=(Type=Cloud, ServerID=Cloud, EnvHostOverride=UE-CloudDataCacheHost, CommandLineHostOverride=CloudDataCacheHost, CommandLineOAuthSecretOverride="CloudDataCacheOAuthSecret", EnvOAuthAccessTokenOverride="UE-CloudDataCacheAccessToken", EnvHttpVersionOverride="UE-CloudDataCacheHttpVersion", CommandLineHttpVersionOverride="CloudDataCacheHttpVersion", EnvOAuthProviderIdentifierOverride="UE-CloudDataCacheOAuthProviderIdentifier", CommandLineOAuthProviderIdentifierOverride="CloudDataCacheOAuthProviderIdentifier")
@@ -2885,9 +2864,7 @@ ZenLocal=(Type=Zen, ServerID=Local)
 ZenShared=(Type=Zen, ServerID=Shared)
 
 [CreateInstalledEnterprisePak]
-Root=(Type=KeyLength, Length=120, Inner=AsyncPut)
-AsyncPut=(Type=AsyncPut, Inner=Hierarchy)
-Hierarchy=(Type=Hierarchical, Inner=PakWrite, Inner=ZenLocal, Inner=Local, Inner=ZenShared, Inner=Shared, Inner=Cloud)
+Root=(Type=Hierarchical, Inner=PakWrite, Inner=ZenLocal, Inner=Local, Inner=ZenShared, Inner=Shared, Inner=Cloud)
 Local=(Type=FileSystem, DeleteOnly=true, UnusedFileAge=8, Path=%ENGINEDIR%DerivedDataCache, EnvPathOverride=UE-LocalDataCachePath, EditorOverrideSetting=LocalDerivedDataCache, CommandLineOverride=LocalDataCachePath)
 Shared=(Type=FileSystem, DeleteUnused=false, ConsiderSlowAt=70, Path=?EpicDDC, EnvPathOverride=UE-SharedDataCachePath, EditorOverrideSetting=SharedDerivedDataCache, CommandLineOverride=SharedDataCachePath)
 Cloud=(Type=Cloud, ServerID=Cloud, EnvHostOverride=UE-CloudDataCacheHost, CommandLineHostOverride=CloudDataCacheHost, CommandLineOAuthSecretOverride="CloudDataCacheOAuthSecret", EnvOAuthAccessTokenOverride="UE-CloudDataCacheAccessToken", EnvHttpVersionOverride="UE-CloudDataCacheHttpVersion", CommandLineHttpVersionOverride="CloudDataCacheHttpVersion", EnvOAuthProviderIdentifierOverride="UE-CloudDataCacheOAuthProviderIdentifier", CommandLineOAuthProviderIdentifierOverride="CloudDataCacheOAuthProviderIdentifier")
@@ -2896,9 +2873,7 @@ ZenLocal=(Type=Zen, ServerID=Local)
 ZenShared=(Type=Zen, ServerID=Shared)
 
 [CreateProjectCache]
-Root=(Type=KeyLength, Length=120, Inner=AsyncPut)
-AsyncPut=(Type=AsyncPut, Inner=Hierarchy)
-Hierarchy=(Type=Hierarchical, Inner=ZenLocal, Inner=Local, Inner=Project, Inner=ZenShared, Inner=Shared, Inner=Cloud)
+Root=(Type=Hierarchical, Inner=ZenLocal, Inner=Local, Inner=Project, Inner=ZenShared, Inner=Shared, Inner=Cloud)
 Local=(Type=FileSystem, DeleteOnly=true, UnusedFileAge=8, Path=%ENGINEDIR%DerivedDataCache, EnvPathOverride=UE-LocalDataCachePath, EditorOverrideSetting=LocalDerivedDataCache, CommandLineOverride=LocalDataCachePath)
 Project=(Type=FileSystem, Flush=true, UnusedFileAge=34, Path=%GAMEDIR%ProjectDerivedData)
 Shared=(Type=FileSystem, DeleteUnused=false, ConsiderSlowAt=70, Path=?EpicDDC, EnvPathOverride=UE-SharedDataCachePath, EditorOverrideSetting=SharedDerivedDataCache, CommandLineOverride=SharedDataCachePath)
