Make sure the ExternalActors commandlet is part of the CommandletsUsingAR set so the Asset Registry is ready when we start parsing actor classes.

diff --git a/Engine/Config/BaseEngine.ini b/Engine/Config/BaseEngine.ini
--- a/Engine/Config/BaseEngine.ini
+++ b/Engine/Config/BaseEngine.ini
@@ -2961,11 +2961,12 @@ GatherOnGameThreadOnly=false
 
 ; Commandlets by default do not SearchAllAssets because many do not need it. A commandlet can be opted in
 ; using CommandletsUsingAR. Format is the shortname of the commandlet, optionally without the "Commandlet" suffix.
++CommandletsUsingAR=CookShaders
 +CommandletsUsingAR=DataValidation
++CommandletsUsingAR=ExternalActors
 +CommandletsUsingAR=GatherText
 +CommandletsUsingAR=GenerateNaniteDisplacedMesh
 +CommandletsUsingAR=WorldPartitionBuilder
-+CommandletsUsingAR=CookShaders
 
 ; AssetDataGathering can be cached in the editor; there are two different caches, one for the directory scan to find all
 ; packages and one for tags and other data read from the packages. The directory scan cache by default is only read
