Net
- Bumped default ServerDesiredSocketReceiveBufferBytes and ClientDesiredSocketReceiveBufferBytes to 262144 bytes.

diff --git a/Engine/Config/BaseEngine.ini b/Engine/Config/BaseEngine.ini
--- a/Engine/Config/BaseEngine.ini
+++ b/Engine/Config/BaseEngine.ini
@@ -1804,6 +1804,8 @@ TimeoutMultiplierForUnoptimizedBuilds=1
 KeepAliveTime=0.2
 MaxClientRate=100000
 MaxInternetClientRate=100000
+ServerDesiredSocketReceiveBufferBytes=262144
+ClientDesiredSocketReceiveBufferBytes=262144
 RelevantTimeout=5.0
 SpawnPrioritySeconds=1.0
 ServerTravelPause=4.0
