DDC: Removed the delete-only local file system cache store and other obsolete cache config

The local file system cache was effectively deprecated in 5.4 in favor of Zen. The delete-only config was to automatically clean up the cache over the following releases.

diff --git a/Engine/Config/BaseEngine.ini b/Engine/Config/BaseEngine.ini
--- a/Engine/Config/BaseEngine.ini
+++ b/Engine/Config/BaseEngine.ini
@@ -2724,7 +2724,7 @@ Shared=(Host=None, Namespace="ue.ddc", EnvHostOverride=UE-ZenSharedDataCacheHost
 
 [DerivedDataBackendGraph]
 ; Configure the hierarchy as the root of the cache. Reads access the nodes in order until satisfied. Writes go to every writable node.
-Root=(Type=Hierarchical, Inner=Pak, Inner=EnginePak, Inner=ZenLocal, Inner=Local, Inner=ZenShared, Inner=Shared, Inner=Cloud)
+Root=(Type=Hierarchical, Inner=Pak, Inner=EnginePak, Inner=ZenLocal, Inner=ZenShared, Inner=Shared, Inner=Cloud)
 ; Configure the local cache. This is a FileSystem cache with the following parameters:
 ;	ReadOnly: Can data be written to this layer. Default is false.
 ;	Touch: Update timestamps on every access instead of only after 1/4 of UnusedFileAge. Default is false.
@@ -2744,7 +2744,7 @@ Root=(Type=Hierarchical, Inner=Pak, Inner=EnginePak, Inner=ZenLocal, Inner=Local
 ;		ValueOnly: Look only for compressed data in the 5.0+ directory layout.
 ;		LegacyOnly: Look only for uncompressed data in the pre-5.0 directory layout.
 ;		ValueWithLegacyFallback: Look in the 5.0+ location but fall back to the pre-5.0 location on miss.
-Local=(Type=FileSystem, DeleteOnly=true, UnusedFileAge=8, Path=%ENGINEDIR%DerivedDataCache, EnvPathOverride=UE-LocalDataCachePath, EditorOverrideSetting=LocalDerivedDataCache, CommandLineOverride=LocalDataCachePath)
+;Local=(Type=FileSystem, DeleteOnly=true, UnusedFileAge=8, Path=%ENGINEDIR%DerivedDataCache, EnvPathOverride=UE-LocalDataCachePath, EditorOverrideSetting=LocalDerivedDataCache, CommandLineOverride=LocalDataCachePath)
 ; Configure the shared cache that is accessed after local.
 Shared=(Type=FileSystem, UnusedFileAge=10, FoldersToClean=10, ConsiderSlowAt=70, Path=?EpicDDC, EnvPathOverride=UE-SharedDataCachePath, EditorOverrideSetting=SharedDerivedDataCache, CommandLineOverride=SharedDataCachePath)
 ; Configure the cloud cache that is accessed after shared.
@@ -2758,8 +2758,7 @@ ZenLocal=(Type=Zen, ServerID=Local)
 ZenShared=(Type=Zen, ServerID=Shared)
 
 [InstalledDerivedDataBackendGraph]
-Root=(Type=Hierarchical, Inner=Pak, Inner=CompressedPak, Inner=EnginePak, Inner=EnterprisePak, Inner=ZenLocal, Inner=Local, Inner=ZenShared, Inner=Shared, Inner=Cloud)
-Local=(Type=FileSystem, DeleteOnly=true, UnusedFileAge=8, Path="%ENGINEVERSIONAGNOSTICUSERDIR%DerivedDataCache", EnvPathOverride=UE-LocalDataCachePath, EditorOverrideSetting=LocalDerivedDataCache)
+Root=(Type=Hierarchical, Inner=Pak, Inner=CompressedPak, Inner=EnginePak, Inner=EnterprisePak, Inner=ZenLocal, Inner=ZenShared, Inner=Shared, Inner=Cloud)
 Shared=(Type=FileSystem, UnusedFileAge=10, FoldersToClean=10, ConsiderSlowAt=70, Path=?EpicDDC, EnvPathOverride=UE-SharedDataCachePath, EditorOverrideSetting=SharedDerivedDataCache, CommandLineOverride=SharedDataCachePath)
 Cloud=(Type=Cloud, ServerID=Cloud, EnvHostOverride=UE-CloudDataCacheHost, CommandLineHostOverride=CloudDataCacheHost, CommandLineOAuthSecretOverride="CloudDataCacheOAuthSecret", EnvOAuthAccessTokenOverride="UE-CloudDataCacheAccessToken", EnvHttpVersionOverride="UE-CloudDataCacheHttpVersion", CommandLineHttpVersionOverride="CloudDataCacheHttpVersion", EnvOAuthProviderIdentifierOverride="UE-CloudDataCacheOAuthProviderIdentifier", CommandLineOAuthProviderIdentifierOverride="CloudDataCacheOAuthProviderIdentifier")
 Pak=(Type=ReadPak, Filename="%GAMEDIR%DerivedDataCache/DDC.ddp")
@@ -2794,8 +2793,7 @@ Root=(Type=Hierarchical, Inner=Zen)
 Zen=(Type=Zen, ServerID=Local)
 
 [ZenDDC]
-Root=(Type=Hierarchical, Inner=Pak, Inner=EnginePak, Inner=ZenLocal, Inner=Local, Inner=ZenShared, Inner=Shared, Inner=Cloud)
-Local=(Type=FileSystem, DeleteOnly=true, UnusedFileAge=8, Path=%ENGINEDIR%DerivedDataCache, EnvPathOverride=UE-LocalDataCachePath, EditorOverrideSetting=LocalDerivedDataCache, CommandLineOverride=LocalDataCachePath)
+Root=(Type=Hierarchical, Inner=Pak, Inner=EnginePak, Inner=ZenLocal, Inner=ZenShared, Inner=Shared, Inner=Cloud)
 ZenLocal=(Type=Zen, ServerID=Local)
 ZenShared=(Type=Zen, ServerID=Shared)
 Shared=(Type=FileSystem, UnusedFileAge=10, FoldersToClean=10, ConsiderSlowAt=70, Path=?EpicDDC, EnvPathOverride=UE-SharedDataCachePath, EditorOverrideSetting=SharedDerivedDataCache, CommandLineOverride=SharedDataCachePath)
@@ -2812,16 +2810,14 @@ Root=(Type=Hierarchical, Inner=ZenLocal)
 ZenLocal=(Type=Zen, ServerID=Local, Sandbox=Cold)
 
 [DerivedDataBackendGraph_Fill_Seattle]
-Root=(Type=Hierarchical, Inner=Pak, Inner=EnginePak, Inner=ZenLocal, Inner=Local, Inner=Seattle)
-Local=(Type=FileSystem, DeleteOnly=true, UnusedFileAge=8, Path=%ENGINEDIR%DerivedDataCache, EnvPathOverride=UE-LocalDataCachePath, EditorOverrideSetting=LocalDerivedDataCache, CommandLineOverride=LocalDataCachePath)
+Root=(Type=Hierarchical, Inner=Pak, Inner=EnginePak, Inner=ZenLocal, Inner=Seattle)
 Seattle=(Type=FileSystem, UnusedFileAge=23, FoldersToClean=10, Path=?EpicSeaDDC, EnvPathOverride=UE-SharedDataCachePath_Seattle)
 Pak=(Type=ReadPak, Filename="%GAMEDIR%DerivedDataCache/DDC.ddp")
 EnginePak=(Type=ReadPak, Filename="%ENGINEDIR%DerivedDataCache/DDC.ddp")
 ZenLocal=(Type=Zen, ServerID=Local)
 
 [NoShared]
-Root=(Type=Hierarchical, Inner=Pak, Inner=ZenLocal, Inner=Local)
-Local=(Type=FileSystem, DeleteOnly=true, UnusedFileAge=8, Path=%ENGINEDIR%DerivedDataCache, EnvPathOverride=UE-LocalDataCachePath, EditorOverrideSetting=LocalDerivedDataCache, CommandLineOverride=LocalDataCachePath)
+Root=(Type=Hierarchical, Inner=Pak, Inner=ZenLocal)
 Pak=(Type=ReadPak, Filename="%GAMEDIR%DerivedDataCache/DDC.ddp")
 ZenLocal=(Type=Zen, ServerID=Local)
 
@@ -2836,8 +2832,7 @@ Root=(Type=Hierarchical, Inner=ZenLocal)
 ZenLocal=(Type=Zen, ServerID=Local, Sandbox=Cold)
 
 [CreatePak]
-Root=(Type=Hierarchical, Inner=PakWrite, Inner=PakRead, Inner=ZenLocal, Inner=Local, Inner=ZenShared, Inner=Shared, Inner=Cloud)
-Local=(Type=FileSystem, DeleteOnly=true, UnusedFileAge=8, Path=%ENGINEDIR%DerivedDataCache, EnvPathOverride=UE-LocalDataCachePath, EditorOverrideSetting=LocalDerivedDataCache, CommandLineOverride=LocalDataCachePath)
+Root=(Type=Hierarchical, Inner=PakWrite, Inner=PakRead, Inner=ZenLocal, Inner=ZenShared, Inner=Shared, Inner=Cloud)
 Shared=(Type=FileSystem, DeleteUnused=false, ConsiderSlowAt=70, Path=?EpicDDC, EnvPathOverride=UE-SharedDataCachePath, EditorOverrideSetting=SharedDerivedDataCache, CommandLineOverride=SharedDataCachePath)
 Cloud=(Type=Cloud, ServerID=Cloud, EnvHostOverride=UE-CloudDataCacheHost, CommandLineHostOverride=CloudDataCacheHost, CommandLineOAuthSecretOverride="CloudDataCacheOAuthSecret", EnvOAuthAccessTokenOverride="UE-CloudDataCacheAccessToken", EnvHttpVersionOverride="UE-CloudDataCacheHttpVersion", CommandLineHttpVersionOverride="CloudDataCacheHttpVersion", EnvOAuthProviderIdentifierOverride="UE-CloudDataCacheOAuthProviderIdentifier", CommandLineOAuthProviderIdentifierOverride="CloudDataCacheOAuthProviderIdentifier")
 PakRead=(Type=ReadPak, Filename="%GAMEDIR%DerivedDataCache/DDC.ddp")
@@ -2846,8 +2841,7 @@ ZenLocal=(Type=Zen, ServerID=Local)
 ZenShared=(Type=Zen, ServerID=Shared)
 
 [CreateInstalledProjectPak]
-Root=(Type=Hierarchical, Inner=EnginePak, Inner=PakWrite, Inner=PakRead, Inner=ZenLocal, Inner=Local, Inner=ZenShared, Inner=Shared, Inner=Cloud)
-Local=(Type=FileSystem, DeleteOnly=true, UnusedFileAge=8, Path=%ENGINEDIR%DerivedDataCache, EnvPathOverride=UE-LocalDataCachePath, EditorOverrideSetting=LocalDerivedDataCache, CommandLineOverride=LocalDataCachePath)
+Root=(Type=Hierarchical, Inner=EnginePak, Inner=PakWrite, Inner=PakRead, Inner=ZenLocal, Inner=ZenShared, Inner=Shared, Inner=Cloud)
 Shared=(Type=FileSystem, DeleteUnused=false, ConsiderSlowAt=70, Path=?EpicDDC, EnvPathOverride=UE-SharedDataCachePath, EditorOverrideSetting=SharedDerivedDataCache, CommandLineOverride=SharedDataCachePath)
 Cloud=(Type=Cloud, ServerID=Cloud, EnvHostOverride=UE-CloudDataCacheHost, CommandLineHostOverride=CloudDataCacheHost, CommandLineOAuthSecretOverride="CloudDataCacheOAuthSecret", EnvOAuthAccessTokenOverride="UE-CloudDataCacheAccessToken", EnvHttpVersionOverride="UE-CloudDataCacheHttpVersion", CommandLineHttpVersionOverride="CloudDataCacheHttpVersion", EnvOAuthProviderIdentifierOverride="UE-CloudDataCacheOAuthProviderIdentifier", CommandLineOAuthProviderIdentifierOverride="CloudDataCacheOAuthProviderIdentifier")
 EnginePak=(Type=ReadPak, Filename=../../../Engine/DerivedDataCache/Compressed.ddp, Compressed=true)
@@ -2857,8 +2851,7 @@ ZenLocal=(Type=Zen, ServerID=Local)
 ZenShared=(Type=Zen, ServerID=Shared)
 
 [CreateInstalledEnginePak]
-Root=(Type=Hierarchical, Inner=PakWrite, Inner=ZenLocal, Inner=Local, Inner=ZenShared, Inner=Shared, Inner=Cloud)
-Local=(Type=FileSystem, DeleteOnly=true, UnusedFileAge=8, Path=%ENGINEDIR%DerivedDataCache, EnvPathOverride=UE-LocalDataCachePath, EditorOverrideSetting=LocalDerivedDataCache, CommandLineOverride=LocalDataCachePath)
+Root=(Type=Hierarchical, Inner=PakWrite, Inner=ZenLocal, Inner=ZenShared, Inner=Shared, Inner=Cloud)
 Shared=(Type=FileSystem, DeleteUnused=false, ConsiderSlowAt=70, Path=?EpicDDC, EnvPathOverride=UE-SharedDataCachePath, EditorOverrideSetting=SharedDerivedDataCache, CommandLineOverride=SharedDataCachePath)
 Cloud=(Type=Cloud, ServerID=Cloud, EnvHostOverride=UE-CloudDataCacheHost, CommandLineHostOverride=CloudDataCacheHost, CommandLineOAuthSecretOverride="CloudDataCacheOAuthSecret", EnvOAuthAccessTokenOverride="UE-CloudDataCacheAccessToken", EnvHttpVersionOverride="UE-CloudDataCacheHttpVersion", CommandLineHttpVersionOverride="CloudDataCacheHttpVersion", EnvOAuthProviderIdentifierOverride="UE-CloudDataCacheOAuthProviderIdentifier", CommandLineOAuthProviderIdentifierOverride="CloudDataCacheOAuthProviderIdentifier")
 PakWrite=(Type=WritePak, Filename="%ENGINEDIR%DerivedDataCache/Compressed.ddp", Compressed=true)
@@ -2866,8 +2859,7 @@ ZenLocal=(Type=Zen, ServerID=Local)
 ZenShared=(Type=Zen, ServerID=Shared)
 
 [CreateInstalledEnterprisePak]
-Root=(Type=Hierarchical, Inner=PakWrite, Inner=ZenLocal, Inner=Local, Inner=ZenShared, Inner=Shared, Inner=Cloud)
-Local=(Type=FileSystem, DeleteOnly=true, UnusedFileAge=8, Path=%ENGINEDIR%DerivedDataCache, EnvPathOverride=UE-LocalDataCachePath, EditorOverrideSetting=LocalDerivedDataCache, CommandLineOverride=LocalDataCachePath)
+Root=(Type=Hierarchical, Inner=PakWrite, Inner=ZenLocal, Inner=ZenShared, Inner=Shared, Inner=Cloud)
 Shared=(Type=FileSystem, DeleteUnused=false, ConsiderSlowAt=70, Path=?EpicDDC, EnvPathOverride=UE-SharedDataCachePath, EditorOverrideSetting=SharedDerivedDataCache, CommandLineOverride=SharedDataCachePath)
 Cloud=(Type=Cloud, ServerID=Cloud, EnvHostOverride=UE-CloudDataCacheHost, CommandLineHostOverride=CloudDataCacheHost, CommandLineOAuthSecretOverride="CloudDataCacheOAuthSecret", EnvOAuthAccessTokenOverride="UE-CloudDataCacheAccessToken", EnvHttpVersionOverride="UE-CloudDataCacheHttpVersion", CommandLineHttpVersionOverride="CloudDataCacheHttpVersion", EnvOAuthProviderIdentifierOverride="UE-CloudDataCacheOAuthProviderIdentifier", CommandLineOAuthProviderIdentifierOverride="CloudDataCacheOAuthProviderIdentifier")
 PakWrite=(Type=WritePak, Filename=../../../Enterprise/DerivedDataCache/Compressed.ddp, Compressed=true)
@@ -2875,8 +2867,7 @@ ZenLocal=(Type=Zen, ServerID=Local)
 ZenShared=(Type=Zen, ServerID=Shared)
 
 [CreateProjectCache]
-Root=(Type=Hierarchical, Inner=ZenLocal, Inner=Local, Inner=Project, Inner=ZenShared, Inner=Shared, Inner=Cloud)
-Local=(Type=FileSystem, DeleteOnly=true, UnusedFileAge=8, Path=%ENGINEDIR%DerivedDataCache, EnvPathOverride=UE-LocalDataCachePath, EditorOverrideSetting=LocalDerivedDataCache, CommandLineOverride=LocalDataCachePath)
+Root=(Type=Hierarchical, Inner=ZenLocal, Inner=Project, Inner=ZenShared, Inner=Shared, Inner=Cloud)
 Project=(Type=FileSystem, Flush=true, UnusedFileAge=34, Path=%GAMEDIR%ProjectDerivedData)
 Shared=(Type=FileSystem, DeleteUnused=false, ConsiderSlowAt=70, Path=?EpicDDC, EnvPathOverride=UE-SharedDataCachePath, EditorOverrideSetting=SharedDerivedDataCache, CommandLineOverride=SharedDataCachePath)
 Cloud=(Type=Cloud, ServerID=Cloud, EnvHostOverride=UE-CloudDataCacheHost, CommandLineHostOverride=CloudDataCacheHost, CommandLineOAuthSecretOverride="CloudDataCacheOAuthSecret", EnvOAuthAccessTokenOverride="UE-CloudDataCacheAccessToken")
