Iris
* Disable log warning for UpdateLevelVisibilityLevelInfo not having an Iris net serializer

#rnx

diff --git a/Engine/Config/BaseEngine.ini b/Engine/Config/BaseEngine.ini
--- a/Engine/Config/BaseEngine.ini
+++ b/Engine/Config/BaseEngine.ini
@@ -1513,6 +1513,7 @@ bEnsureAllClassesAreFullyPushModel=false
 +SupportsStructNetSerializerList=(StructName=RootMotionSource_MoveToDynamicForce)
 +SupportsStructNetSerializerList=(StructName=RootMotionSource_MoveToForce)
 +SupportsStructNetSerializerList=(StructName=RootMotionSource_RadialForce)
++SupportsStructNetSerializerList=(StructName=UpdateLevelVisibilityLevelInfo)
 
 [/Script/IrisCore.NetTokenTypeIdConfig]
 +ReservedTypeIds=(StoreTypeName="StringTokenStore", TypeID=0)
