DDC: Migrate config to the new format

#rnx

diff --git a/Engine/Config/BaseEngine.ini b/Engine/Config/BaseEngine.ini
--- a/Engine/Config/BaseEngine.ini
+++ b/Engine/Config/BaseEngine.ini
@@ -2674,9 +2674,96 @@ HostName=localhost
 Port=8558
 
 ; Each key defines settings for an Unreal Zen Storage or Unreal Cloud DDC server.
-; The key is used in the ServerID field of a ZenLocal/ZenShared/Cloud node in a graph for DDC.
-; 
-; For an Unreal Cloud DDC node, the configurable parameters are:
+; The key is used in the ServerID field of a Type=Zen or Type=Cloud node in a cache graph.
+; See configurable parameters below in [DerivedDataCacheStores].
+[StorageServers]
+Local=(Namespace="ue.ddc")
+Shared=(Host=None, Namespace="ue.ddc", EnvHostOverride=UE-ZenSharedDataCacheHost, CommandLineHostOverride=ZenSharedDataCacheHost, DeactivateAt=60)
+Cloud=(Host=None, Namespace="ue.ddc")
+
+; Derived Data Cache Graphs
+;
+; A graph is defined by a name and an ordered list of cache stores.
+; Reads access every readable cache store in order until satisfied.
+; Writes go to every writable cache store.
+;
+; "Default" is the default graph for source builds.
+; "Installed" is the default graph for installed builds.
+;
+; A different graph can be specified on the command line:
+;   -DDC=GraphName
+; A custom graph can be specified on the command line too:
+;   -DDC=(Local,Cloud(ReadOnly=true))
+; The DDC.Graph console variable works the same as -DDC=...
+;
+; A graph can rename cache stores to provide consistent naming for the command line, logging, telemetry, etc.:
+;   Custom=(Local=CustomLocal)
+; A graph can inject config overrides for a referenced cache store:
+;   Custom=(Local, Cloud(ReadOnly=true))
+; A graph can both rename and inject overrides for a referenced cache store:
+;   Custom=(Local, Cloud=CustomCloud(ReadOnly=true))
+; A graph can replace the config for a cache store:
+;   Custom=(Local=(Type=Zen, Namespace="custom.ddc"))
+;
+; A graph can be deprecated to generate a warning when it is used:
+;   Custom=(Local, Shared, Cloud, Deprecated="Use Default.")
+[DerivedDataCacheGraphs]
+Default=(ProjectPak, EnginePak, ZenLocal, Local, ZenShared, Shared, Cloud)
+Installed=(ProjectPak, InstalledProjectPak, EnginePak=InstalledEnginePak, EnterprisePak, ZenLocal, Local=InstalledLocal, ZenShared, Shared, Cloud)
+
+NoShared=(ProjectPak, EnginePak, ZenLocal)
+NoZenLocalFallback=(ProjectPak, EnginePak, Local, ZenShared, Shared, Cloud)
+InstalledNoZenLocalFallback=(ProjectPak, InstalledProjectPak, EnginePak=InstalledEnginePak, Local=InstalledLocal, ZenShared, Shared, Cloud)
+
+Cold=(ZenLocal(Sandbox=Cold, Flush=true)) ; An empty standalone cache that does not affect the existing local cache.
+Warm=(ZenLocal(Sandbox=Cold))             ; Use the same standalone cache as Cold without clearing it before use.
+
+CreateProjectPak=(WriteProjectPak, ProjectPak, ZenLocal, ZenShared, Shared, Cloud)
+CreateInstalledProjectPak=(InstalledEnginePak, WriteInstalledProjectPak, InstalledProjectPak, ZenLocal, ZenShared, Shared, Cloud)
+CreateEnginePak=(WriteEnginePak, EnginePak, ZenLocal, ZenShared, Shared, Cloud)
+CreateInstalledEnginePak=(WriteInstalledEnginePak, InstalledEnginePak, ZenLocal, ZenShared, Shared, Cloud)
+CreateInstalledEnterprisePak=(WriteEnterprisePak, EnterprisePak, ZenLocal, ZenShared, Shared, Cloud)
+
+; Provide the original default names until their usage as a fallback has been removed.
+DerivedDataBackendGraph=(ProjectPak, EnginePak, ZenLocal, Local, ZenShared, Shared, Cloud)
+InstalledDerivedDataBackendGraph=(ProjectPak, InstalledProjectPak, EnginePak=InstalledEnginePak, EnterprisePak, ZenLocal, Local=InstalledLocal, ZenShared, Shared, Cloud)
+
+; Deprecated Graphs
+CreatePak=(WriteProjectPak, ProjectPak, ZenLocal, ZenShared, Shared, Cloud, Deprecated="Use CreateProjectPak.")
+ZenDDC=(Hierarchical, ProjectPak, EnginePak, ZenLocal, ZenShared, Shared, Cloud, Deprecated="Use Default.")
+ZenOnlyDDC=(ZenLocal, Deprecated="Use -DDC=(ZenLocal).")
+ZenColdDDC=(ZenLocal(Sandbox=Cold, Flush=true), Deprecated="Use Cold.")
+ZenWarmDDC=(ZenLocal(Sandbox=Cold), Deprecated="Use Warm.")
+CreateProjectCache=(ZenLocal, Project=Local(Path=%GAMEDIR%ProjectDerivedData), ZenShared, Shared, Cloud, Deprecated="Graph is no longer supported.")
+DerivedDataBackendGraph_Fill_Seattle=(Hierarchical, ProjectPak, EnginePak, ZenLocal, Deprecated="Graph is no longer supported.")
+
+; Derived Data Cache Stores
+;
+; A store is defined by a name, a type, and the config for a store of that type.
+; Available types are: Zen, Cloud, FileSystem, Pak (ReadPak/WritePak).
+;
+; A store can be based on another store to avoid duplication of parameters.
+;   Custom=(Base=ZenLocal, Sandbox=Custom)
+;
+; A store can be deprecated to generate a warning when it is used:
+;   Custom=(Base=ZenLocal, Deprecated="Use ZenLocal.")
+[DerivedDataCacheStores]
+
+; Type=Zen is for Unreal Zen Storage. Available parameters are:
+;   Host: The URI for the server to connect to. Locally autolaunched Zen instances can exclude this, e.g, http://sharedzen.example.com:8558.
+;   EnvHostOverride: The name of an environment variable that can override the host, e.g., UE-ZenSharedDataCacheHost. Optional.
+;   CommandLineHostOverride: The name of a command line switch that can override the host, e.g., ZenSharedDataCacheHost. Optional.
+;   Namespace: The cache namespace to use on the server. Required.
+;   Sandbox: Whether this cache usage should be isolated into a named instance distinct from the normal instance. e.g., Cold. Optional.
+;   Flush: Whether this cache usage should be deleted upon startup. Only applies when using a sandbox. Default is false. Optional.
+;   DeactivateAt: A latency threshold (in milliseconds) which when exceeded causes the node to be disabled. e.g., 90. Optional.
+;   Local: Whether to force the flags for the node to include the Local flag. Default is determined by the host URL. Optional.
+;   Remote: Whether to force the flags for the node to include the Remote flag. Default is determined by the host URL. Optional.
+;   ReadOnly: Whether the cache is treated as read-only. Default is false. Optional.
+ZenLocal=(Type=Zen, ServerID=Local)
+ZenShared=(Type=Zen, ServerID=Shared)
+
+; Type=Cloud is for Unreal Cloud DDC. Available parameters are:
 ;   Host: The URI for the server to connect to, e.g., https://jupiter.example.com.
 ;   EnvHostOverride: The name of an environment variable that can override the host, e.g., UE-CloudDataCacheHost.
 ;   CommandLineHostOverride: The name of a command line switch that can override the host, e.g., CloudDataCacheHost.
@@ -2686,47 +2773,22 @@ Port=8558
 ;   OAuthProvider: The URI for the OAuth provider to use for authentication. Optional.
 ;   OAuthClientId: The client ID to use with the OAuth provider. Optional.
 ;   OAuthSecret: The secret to use with the OAuth provider. Optional. May start with file:// to load the secret from a file.
+;   EnvOAuthSecretOverride: The name of an environment variable that can override the OAuth secret, e.g., UE-CloudDataCacheOAuthSecret. Optional.
 ;   CommandLineOAuthSecretOverride: The name of a command line switch that can override the OAuth secret, e.g., CloudDataCacheOAuthSecret. Optional.
 ;   OAuthScope: The scope to use with the OAuth provider. Default is cache_access.
 ;   OAuthProviderIdentifier: The provider ID when using OIDC. Optional.
-;   EnvOAuthAccessTokenOverride: The name of a command line switch that can provide the OAuth access token to use, e.g., UE-CloudDataCacheAccessToken. Optional.
+;   EnvOAuthAccessTokenOverride: The name of an environment variable that can provide the OAuth access token to use, e.g., UE-CloudDataCacheAccessToken. Optional.
+;   CommandLineOAuthAccessTokenOverride: The name of a command line switch that can provide the OAuth access token to use, e.g., CloudDataCacheAccessToken. Optional.
+;   AuthScheme: Default is Bearer.
 ;   OAuthPinnedPublicKeys: A semicolon-delimited list of pinned public keys for the OAuth provider. Optional.
 ;   ReadOnly: Whether the cache is treated as read-only. Default is false.
-;   HttpVersion: An override for which http version to use
+;   WriteOnly: Whether the cache is treated as write-only. Default is false.
+;   HttpVersion: An override for which http version to use. Values: http1.0, http1.1, http2, http2-only.
 ;   EnvHttpVersionOverride: An environment variable that if set will be used instead of HttpVersion.
 ;   CommandLineHttpVersionOverride: A command line argument used in preference to the HttpVersion / EnvHttpVersionOverride settings
-; 
-; For an Unreal Zen Storage node, the configurable parameters are:
-;   Host: The URI for the server to connect to. Locally autolaunched Zen instances can exclude this, e.g, http://sharedzen.example.com:8558.
-;   EnvHostOverride: The name of an environment variable that can override the host, e.g., UE-ZenSharedDataCacheHost. Optional.
-;   CommandLineHostOverride: The name of a command line switch that can override the host, e.g., ZenSharedDataCacheHost. Optional.
-;   Namespace: The cache namespace to use on the server. Required.
-;   Sandbox: Whether this cache usage should be isolated into a named instance distinct from the normal instance. e.g., Cold. Optional.
-;   Flush: Whether this cache usage should be deleted upon startup. Only applies when using a sandbox. Default is false. Optional.
-;   DeactivateAt: A latency threshold (in milliseconds) which when exceeded causes the node to be disabled. e.g., 90. Optional.
-;   Local: Whether to force the flags for the node to include the Local flag. Default is determined by the host URL. Optional.
-;   Remote: Whether to force the flags for the node to include the Remote flag. Default is determined by the host URL. Optional.
-;   ReadOnly: Whether the cache is treated as read-only. Default is false. Optional.
-[StorageServers]
-Cloud=(Host=None, Namespace="ue.ddc")
-Local=(Namespace="ue.ddc")
-Shared=(Host=None, Namespace="ue.ddc", EnvHostOverride=UE-ZenSharedDataCacheHost, CommandLineHostOverride=ZenSharedDataCacheHost, DeactivateAt=60)
+Cloud=(Type=Cloud, ServerID=Cloud, EnvHostOverride=UE-CloudDataCacheHost, CommandLineHostOverride=CloudDataCacheHost, CommandLineOAuthSecretOverride="CloudDataCacheOAuthSecret", EnvOAuthAccessTokenOverride="UE-CloudDataCacheAccessToken", EnvHttpVersionOverride="UE-CloudDataCacheHttpVersion", CommandLineHttpVersionOverride="CloudDataCacheHttpVersion", EnvOAuthProviderIdentifierOverride="UE-CloudDataCacheOAuthProviderIdentifier", CommandLineOAuthProviderIdentifierOverride="CloudDataCacheOAuthProviderIdentifier")
 
-; Derived Data backend graphs
-; Each of the below backend graph sections contains a set of nodes used to create derived data backend graph.
-; DerivedDataBackendGraph is the default graph for source builds, and InstalledDerivedDataBackendGraph is the default for installed builds.
-; Others can be specified on the command line using:
-;     -DDC=GraphSectionName     (for example: -DDC=NoShared)
-; Each graph should start with a 'Root' node. The names of the other nodes are not predefined.
-; Supported node types are: Hierarchical, FileSystem, Zen, Cloud, S3, ReadPak, WritePak
-; The order that nodes are defined in is not relevant.
-; FileSystem backends can be disabled by setting the corresponding environment variable to 'None'. E.g., UE-SharedDataCachePath=None
-; Cloud backends can be disabled by setting the corresponding environment variable to 'None'. E.g., UE-CloudDataCachePath=None
-
-[DerivedDataBackendGraph]
-; Configure the hierarchy as the root of the cache. Reads access the nodes in order until satisfied. Writes go to every writable node.
-Root=(Type=Hierarchical, Inner=Pak, Inner=EnginePak, Inner=ZenLocal, Inner=ZenShared, Inner=Shared, Inner=Cloud)
-; Configure the local cache. This is a FileSystem cache with the following parameters:
+; Type=FileSystem is for a file system cache. Available parameters are:
 ;	ReadOnly: Can data be written to this layer. Default is false.
 ;	Touch: Update timestamps on every access instead of only after 1/4 of UnusedFileAge. Default is false.
 ;	Clean: Perform a cleanup of old files on launch. Default is false.
@@ -2741,139 +2803,29 @@ Root=(Type=Hierarchical, Inner=Pak, Inner=EnginePak, Inner=ZenLocal, Inner=ZenSh
 ;	EnvPathOverride: An environment variable that if set will be used instead of Path. E.g., UE-LocalDataCachePath=D:\DDC. ('None' disables the DDC)
 ;	CommandLineOverride: A command line argument used in preference to the Path / EnvPathOverride settings. E.g., -SharedDataCachePath=\\someshare\folder
 ;	EditorOverrideSetting: Editor user setting that overrides the Path / EnvPathOverride line values.
-;	LegacyMode: Where to look for legacy cache keys in storage. One of ValueOnly, ValueWithLegacyFallback, LegacyOnly.
-;		ValueOnly: Look only for compressed data in the 5.0+ directory layout.
-;		LegacyOnly: Look only for uncompressed data in the pre-5.0 directory layout.
-;		ValueWithLegacyFallback: Look in the 5.0+ location but fall back to the pre-5.0 location on miss.
-;Local=(Type=FileSystem, DeleteOnly=true, UnusedFileAge=8, Path=%ENGINEDIR%DerivedDataCache, EnvPathOverride=UE-LocalDataCachePath, EditorOverrideSetting=LocalDerivedDataCache, CommandLineOverride=LocalDataCachePath)
-; Configure the shared cache that is accessed after local.
-Shared=(Type=FileSystem, UnusedFileAge=10, FoldersToClean=10, ConsiderSlowAt=70, Path=?EpicDDC, EnvPathOverride=UE-SharedDataCachePath, EditorOverrideSetting=SharedDerivedDataCache, CommandLineOverride=SharedDataCachePath)
-; Configure the cloud cache that is accessed after shared.
-Cloud=(Type=Cloud, ServerID=Cloud, EnvHostOverride=UE-CloudDataCacheHost, CommandLineHostOverride=CloudDataCacheHost, CommandLineOAuthSecretOverride="CloudDataCacheOAuthSecret", EnvOAuthAccessTokenOverride="UE-CloudDataCacheAccessToken", EnvHttpVersionOverride="UE-CloudDataCacheHttpVersion", CommandLineHttpVersionOverride="CloudDataCacheHttpVersion", EnvOAuthProviderIdentifierOverride="UE-CloudDataCacheOAuthProviderIdentifier", CommandLineOAuthProviderIdentifierOverride="CloudDataCacheOAuthProviderIdentifier")
-; Configure a Project Pak node. This is a pre-generated DDC data file for the project that can be distributed to reduce runtime fetches/generation
-; See documentation for how to create a DDP via the DerivedDataCache commandlet
-Pak=(Type=ReadPak, Filename="%GAMEDIR%DerivedDataCache/DDC.ddp")
-; Configure a Project Pak node. This is a pre-generated DDC data file for the engine that can be distributed to reduce runtime fetches/generation
-EnginePak=(Type=ReadPak, Filename="%ENGINEDIR%DerivedDataCache/DDC.ddp")
-ZenLocal=(Type=Zen, ServerID=Local)
-ZenShared=(Type=Zen, ServerID=Shared)
-
-[InstalledDerivedDataBackendGraph]
-Root=(Type=Hierarchical, Inner=Pak, Inner=CompressedPak, Inner=EnginePak, Inner=EnterprisePak, Inner=ZenLocal, Inner=ZenShared, Inner=Shared, Inner=Cloud)
-Shared=(Type=FileSystem, UnusedFileAge=10, FoldersToClean=10, ConsiderSlowAt=70, Path=?EpicDDC, EnvPathOverride=UE-SharedDataCachePath, EditorOverrideSetting=SharedDerivedDataCache, CommandLineOverride=SharedDataCachePath)
-Cloud=(Type=Cloud, ServerID=Cloud, EnvHostOverride=UE-CloudDataCacheHost, CommandLineHostOverride=CloudDataCacheHost, CommandLineOAuthSecretOverride="CloudDataCacheOAuthSecret", EnvOAuthAccessTokenOverride="UE-CloudDataCacheAccessToken", EnvHttpVersionOverride="UE-CloudDataCacheHttpVersion", CommandLineHttpVersionOverride="CloudDataCacheHttpVersion", EnvOAuthProviderIdentifierOverride="UE-CloudDataCacheOAuthProviderIdentifier", CommandLineOAuthProviderIdentifierOverride="CloudDataCacheOAuthProviderIdentifier")
-Pak=(Type=ReadPak, Filename="%GAMEDIR%DerivedDataCache/DDC.ddp")
-CompressedPak=(Type=ReadPak, Filename="%GAMEDIR%DerivedDataCache/Compressed.ddp", Compressed=true)
-EnginePak=(Type=ReadPak, Filename=../../../Engine/DerivedDataCache/Compressed.ddp, Compressed=true)
-EnterprisePak=(Type=ReadPak, Filename=../../../Enterprise/DerivedDataCache/Compressed.ddp, Compressed=true)
-ZenLocal=(Type=Zen, ServerID=Local)
-ZenShared=(Type=Zen, ServerID=Shared)
-
-[NoZenLocalFallback]
-Root=(Type=Hierarchical, Inner=Pak, Inner=EnginePak, Inner=Local, Inner=ZenShared, Inner=Shared, Inner=Cloud)
 Local=(Type=FileSystem, UnusedFileAge=34, PromptIfMissing=true, Path=%ENGINEDIR%DerivedDataCache, EnvPathOverride=UE-LocalDataCachePath, EditorOverrideSetting=LocalDerivedDataCache, CommandLineOverride=LocalDataCachePath)
+InstalledLocal=(Base=Local, Path="%ENGINEVERSIONAGNOSTICUSERDIR%DerivedDataCache")
 Shared=(Type=FileSystem, UnusedFileAge=10, FoldersToClean=10, ConsiderSlowAt=70, Path=?EpicDDC, EnvPathOverride=UE-SharedDataCachePath, EditorOverrideSetting=SharedDerivedDataCache, CommandLineOverride=SharedDataCachePath)
-Cloud=(Type=Cloud, ServerID=Cloud, EnvHostOverride=UE-CloudDataCacheHost, CommandLineHostOverride=CloudDataCacheHost, CommandLineOAuthSecretOverride="CloudDataCacheOAuthSecret", EnvOAuthAccessTokenOverride="UE-CloudDataCacheAccessToken", EnvHttpVersionOverride="UE-CloudDataCacheHttpVersion", CommandLineHttpVersionOverride="CloudDataCacheHttpVersion", EnvOAuthProviderIdentifierOverride="UE-CloudDataCacheOAuthProviderIdentifier", CommandLineOAuthProviderIdentifierOverride="CloudDataCacheOAuthProviderIdentifier")
-Pak=(Type=ReadPak, Filename="%GAMEDIR%DerivedDataCache/DDC.ddp")
-EnginePak=(Type=ReadPak, Filename="%ENGINEDIR%DerivedDataCache/DDC.ddp")
-ZenShared=(Type=Zen, ServerID=Shared)
-
-[InstalledNoZenLocalFallback]
-Root=(Type=Hierarchical, Inner=Pak, Inner=CompressedPak, Inner=EnginePak, Inner=EnterprisePak, Inner=Local, Inner=ZenShared, Inner=Shared, Inner=Cloud)
-Local=(Type=FileSystem, UnusedFileAge=34, Path="%ENGINEVERSIONAGNOSTICUSERDIR%DerivedDataCache", EnvPathOverride=UE-LocalDataCachePath, EditorOverrideSetting=LocalDerivedDataCache)
-Shared=(Type=FileSystem, UnusedFileAge=10, FoldersToClean=10, ConsiderSlowAt=70, Path=?EpicDDC, EnvPathOverride=UE-SharedDataCachePath, EditorOverrideSetting=SharedDerivedDataCache, CommandLineOverride=SharedDataCachePath)
-Cloud=(Type=Cloud, ServerID=Cloud, EnvHostOverride=UE-CloudDataCacheHost, CommandLineHostOverride=CloudDataCacheHost, CommandLineOAuthSecretOverride="CloudDataCacheOAuthSecret", EnvOAuthAccessTokenOverride="UE-CloudDataCacheAccessToken", EnvHttpVersionOverride="UE-CloudDataCacheHttpVersion", CommandLineHttpVersionOverride="CloudDataCacheHttpVersion", EnvOAuthProviderIdentifierOverride="UE-CloudDataCacheOAuthProviderIdentifier", CommandLineOAuthProviderIdentifierOverride="CloudDataCacheOAuthProviderIdentifier")
-Pak=(Type=ReadPak, Filename="%GAMEDIR%DerivedDataCache/DDC.ddp")
-CompressedPak=(Type=ReadPak, Filename="%GAMEDIR%DerivedDataCache/Compressed.ddp", Compressed=true)
-EnginePak=(Type=ReadPak, Filename=../../../Engine/DerivedDataCache/Compressed.ddp, Compressed=true)
-EnterprisePak=(Type=ReadPak, Filename=../../../Enterprise/DerivedDataCache/Compressed.ddp, Compressed=true)
-ZenShared=(Type=Zen, ServerID=Shared)
-
-[ZenOnlyDDC]
-Root=(Type=Hierarchical, Inner=Zen)
-Zen=(Type=Zen, ServerID=Local)
-
-[ZenDDC]
-Root=(Type=Hierarchical, Inner=Pak, Inner=EnginePak, Inner=ZenLocal, Inner=ZenShared, Inner=Shared, Inner=Cloud)
-ZenLocal=(Type=Zen, ServerID=Local)
-ZenShared=(Type=Zen, ServerID=Shared)
-Shared=(Type=FileSystem, UnusedFileAge=10, FoldersToClean=10, ConsiderSlowAt=70, Path=?EpicDDC, EnvPathOverride=UE-SharedDataCachePath, EditorOverrideSetting=SharedDerivedDataCache, CommandLineOverride=SharedDataCachePath)
-Cloud=(Type=Cloud, ServerID=Cloud, EnvHostOverride=UE-CloudDataCacheHost, CommandLineHostOverride=CloudDataCacheHost, CommandLineOAuthSecretOverride="CloudDataCacheOAuthSecret", EnvOAuthAccessTokenOverride="UE-CloudDataCacheAccessToken", EnvHttpVersionOverride="UE-CloudDataCacheHttpVersion", CommandLineHttpVersionOverride="CloudDataCacheHttpVersion", EnvOAuthProviderIdentifierOverride="UE-CloudDataCacheOAuthProviderIdentifier", CommandLineOAuthProviderIdentifierOverride="CloudDataCacheOAuthProviderIdentifier")
-Pak=(Type=ReadPak, Filename="%GAMEDIR%DerivedDataCache/DDC.ddp")
-EnginePak=(Type=ReadPak, Filename="%ENGINEDIR%DerivedDataCache/DDC.ddp")
-
-[ZenColdDDC]
-Root=(Type=Hierarchical, Inner=ZenLocal)
-ZenLocal=(Type=Zen, ServerID=Local, Sandbox=Cold, Flush=true)
-
-[ZenWarmDDC]
-Root=(Type=Hierarchical, Inner=ZenLocal)
-ZenLocal=(Type=Zen, ServerID=Local, Sandbox=Cold)
 
-[DerivedDataBackendGraph_Fill_Seattle]
-Root=(Type=Hierarchical, Inner=Pak, Inner=EnginePak, Inner=ZenLocal, Inner=Seattle)
-Seattle=(Type=FileSystem, UnusedFileAge=23, FoldersToClean=10, Path=?EpicSeaDDC, EnvPathOverride=UE-SharedDataCachePath_Seattle)
-Pak=(Type=ReadPak, Filename="%GAMEDIR%DerivedDataCache/DDC.ddp")
+; Type=ReadPak is for reading a pak file cache. Available parameters are:
+;   Filename: Path to the pak file to read.
+; Configure a Project Pak node. This is a pre-generated DDC data file for the project that can be distributed to reduce runtime fetches/generation.
+; See documentation for how to create a DDP via the DerivedDataCache commandlet.
+Pak=(Type=ReadPak, Filename="%GAMEDIR%DerivedDataCache/DDC.ddp", Deprecated="Use ProjectPak.")
+ProjectPak=(Type=ReadPak, Filename="%GAMEDIR%DerivedDataCache/DDC.ddp")
+InstalledProjectPak=(Type=ReadPak, Filename="%GAMEDIR%DerivedDataCache/Compressed.ddp", Compressed=true)
 EnginePak=(Type=ReadPak, Filename="%ENGINEDIR%DerivedDataCache/DDC.ddp")
-ZenLocal=(Type=Zen, ServerID=Local)
-
-[NoShared]
-Root=(Type=Hierarchical, Inner=Pak, Inner=ZenLocal)
-Pak=(Type=ReadPak, Filename="%GAMEDIR%DerivedDataCache/DDC.ddp")
-ZenLocal=(Type=Zen, ServerID=Local)
-
-; Creates a completely new standalone cache so as to not destroy a valuable populated local cache
-[Cold]
-Root=(Type=Hierarchical, Inner=ZenLocal)
-ZenLocal=(Type=Zen, ServerID=Local, Sandbox=Cold, Flush=true)
-
-; Inherits the same standalone cache from [Cold], for incremental performance measurement purposes
-[Warm]
-Root=(Type=Hierarchical, Inner=ZenLocal)
-ZenLocal=(Type=Zen, ServerID=Local, Sandbox=Cold)
-
-[CreatePak]
-Root=(Type=Hierarchical, Inner=PakWrite, Inner=PakRead, Inner=ZenLocal, Inner=ZenShared, Inner=Shared, Inner=Cloud)
-Shared=(Type=FileSystem, DeleteUnused=false, ConsiderSlowAt=70, Path=?EpicDDC, EnvPathOverride=UE-SharedDataCachePath, EditorOverrideSetting=SharedDerivedDataCache, CommandLineOverride=SharedDataCachePath)
-Cloud=(Type=Cloud, ServerID=Cloud, EnvHostOverride=UE-CloudDataCacheHost, CommandLineHostOverride=CloudDataCacheHost, CommandLineOAuthSecretOverride="CloudDataCacheOAuthSecret", EnvOAuthAccessTokenOverride="UE-CloudDataCacheAccessToken", EnvHttpVersionOverride="UE-CloudDataCacheHttpVersion", CommandLineHttpVersionOverride="CloudDataCacheHttpVersion", EnvOAuthProviderIdentifierOverride="UE-CloudDataCacheOAuthProviderIdentifier", CommandLineOAuthProviderIdentifierOverride="CloudDataCacheOAuthProviderIdentifier")
-PakRead=(Type=ReadPak, Filename="%GAMEDIR%DerivedDataCache/DDC.ddp")
-PakWrite=(Type=WritePak, Filename="%GAMEDIR%DerivedDataCache/DDC.ddp")
-ZenLocal=(Type=Zen, ServerID=Local)
-ZenShared=(Type=Zen, ServerID=Shared)
-
-[CreateInstalledProjectPak]
-Root=(Type=Hierarchical, Inner=EnginePak, Inner=PakWrite, Inner=PakRead, Inner=ZenLocal, Inner=ZenShared, Inner=Shared, Inner=Cloud)
-Shared=(Type=FileSystem, DeleteUnused=false, ConsiderSlowAt=70, Path=?EpicDDC, EnvPathOverride=UE-SharedDataCachePath, EditorOverrideSetting=SharedDerivedDataCache, CommandLineOverride=SharedDataCachePath)
-Cloud=(Type=Cloud, ServerID=Cloud, EnvHostOverride=UE-CloudDataCacheHost, CommandLineHostOverride=CloudDataCacheHost, CommandLineOAuthSecretOverride="CloudDataCacheOAuthSecret", EnvOAuthAccessTokenOverride="UE-CloudDataCacheAccessToken", EnvHttpVersionOverride="UE-CloudDataCacheHttpVersion", CommandLineHttpVersionOverride="CloudDataCacheHttpVersion", EnvOAuthProviderIdentifierOverride="UE-CloudDataCacheOAuthProviderIdentifier", CommandLineOAuthProviderIdentifierOverride="CloudDataCacheOAuthProviderIdentifier")
-EnginePak=(Type=ReadPak, Filename=../../../Engine/DerivedDataCache/Compressed.ddp, Compressed=true)
-PakRead=(Type=ReadPak, Filename="%GAMEDIR%DerivedDataCache/Compressed.ddp", Compressed=true)
-PakWrite=(Type=WritePak, Filename="%GAMEDIR%DerivedDataCache/Compressed.ddp", Compressed=true)
-ZenLocal=(Type=Zen, ServerID=Local)
-ZenShared=(Type=Zen, ServerID=Shared)
-
-[CreateInstalledEnginePak]
-Root=(Type=Hierarchical, Inner=PakWrite, Inner=ZenLocal, Inner=ZenShared, Inner=Shared, Inner=Cloud)
-Shared=(Type=FileSystem, DeleteUnused=false, ConsiderSlowAt=70, Path=?EpicDDC, EnvPathOverride=UE-SharedDataCachePath, EditorOverrideSetting=SharedDerivedDataCache, CommandLineOverride=SharedDataCachePath)
-Cloud=(Type=Cloud, ServerID=Cloud, EnvHostOverride=UE-CloudDataCacheHost, CommandLineHostOverride=CloudDataCacheHost, CommandLineOAuthSecretOverride="CloudDataCacheOAuthSecret", EnvOAuthAccessTokenOverride="UE-CloudDataCacheAccessToken", EnvHttpVersionOverride="UE-CloudDataCacheHttpVersion", CommandLineHttpVersionOverride="CloudDataCacheHttpVersion", EnvOAuthProviderIdentifierOverride="UE-CloudDataCacheOAuthProviderIdentifier", CommandLineOAuthProviderIdentifierOverride="CloudDataCacheOAuthProviderIdentifier")
-PakWrite=(Type=WritePak, Filename="%ENGINEDIR%DerivedDataCache/Compressed.ddp", Compressed=true)
-ZenLocal=(Type=Zen, ServerID=Local)
-ZenShared=(Type=Zen, ServerID=Shared)
-
-[CreateInstalledEnterprisePak]
-Root=(Type=Hierarchical, Inner=PakWrite, Inner=ZenLocal, Inner=ZenShared, Inner=Shared, Inner=Cloud)
-Shared=(Type=FileSystem, DeleteUnused=false, ConsiderSlowAt=70, Path=?EpicDDC, EnvPathOverride=UE-SharedDataCachePath, EditorOverrideSetting=SharedDerivedDataCache, CommandLineOverride=SharedDataCachePath)
-Cloud=(Type=Cloud, ServerID=Cloud, EnvHostOverride=UE-CloudDataCacheHost, CommandLineHostOverride=CloudDataCacheHost, CommandLineOAuthSecretOverride="CloudDataCacheOAuthSecret", EnvOAuthAccessTokenOverride="UE-CloudDataCacheAccessToken", EnvHttpVersionOverride="UE-CloudDataCacheHttpVersion", CommandLineHttpVersionOverride="CloudDataCacheHttpVersion", EnvOAuthProviderIdentifierOverride="UE-CloudDataCacheOAuthProviderIdentifier", CommandLineOAuthProviderIdentifierOverride="CloudDataCacheOAuthProviderIdentifier")
-PakWrite=(Type=WritePak, Filename=../../../Enterprise/DerivedDataCache/Compressed.ddp, Compressed=true)
-ZenLocal=(Type=Zen, ServerID=Local)
-ZenShared=(Type=Zen, ServerID=Shared)
+InstalledEnginePak=(Type=ReadPak, Filename=../../../Engine/DerivedDataCache/Compressed.ddp, Compressed=true)
+EnterprisePak=(Type=ReadPak, Filename=../../../Enterprise/DerivedDataCache/Compressed.ddp, Compressed=true)
 
-[CreateProjectCache]
-Root=(Type=Hierarchical, Inner=ZenLocal, Inner=Project, Inner=ZenShared, Inner=Shared, Inner=Cloud)
-Project=(Type=FileSystem, Flush=true, UnusedFileAge=34, Path=%GAMEDIR%ProjectDerivedData)
-Shared=(Type=FileSystem, DeleteUnused=false, ConsiderSlowAt=70, Path=?EpicDDC, EnvPathOverride=UE-SharedDataCachePath, EditorOverrideSetting=SharedDerivedDataCache, CommandLineOverride=SharedDataCachePath)
-Cloud=(Type=Cloud, ServerID=Cloud, EnvHostOverride=UE-CloudDataCacheHost, CommandLineHostOverride=CloudDataCacheHost, CommandLineOAuthSecretOverride="CloudDataCacheOAuthSecret", EnvOAuthAccessTokenOverride="UE-CloudDataCacheAccessToken")
-ZenLocal=(Type=Zen, ServerID=Local)
-ZenShared=(Type=Zen, ServerID=Shared)
+; Type=WritePak is for writing a pak file cache. Available parameters are:
+;   Filename: Path to the pak file to write.
+;   Compressed: True to compress the data more than the default.
+WriteProjectPak=(Type=WritePak, Filename="%GAMEDIR%DerivedDataCache/DDC.ddp")
+WriteInstalledProjectPak=(Type=WritePak, Filename="%GAMEDIR%DerivedDataCache/Compressed.ddp", Compressed=true)
+WriteEnginePak=(Type=WritePak, Filename="%ENGINEDIR%DerivedDataCache/DDC.ddp")
+WriteInstalledEnginePak=(Type=WritePak, Filename="%ENGINEDIR%DerivedDataCache/Compressed.ddp", Compressed=true)
+WriteEnterprisePak=(Type=WritePak, Filename=../../../Enterprise/DerivedDataCache/Compressed.ddp, Compressed=true)
 
 [VirtualTextureChunkDDCCache]
 UnusedFileAge=34
diff --git a/Engine/Programs/ShaderCompileWorker/Config/DefaultEngine.ini b/Engine/Programs/ShaderCompileWorker/Config/DefaultEngine.ini
--- a/Engine/Programs/ShaderCompileWorker/Config/DefaultEngine.ini
+++ b/Engine/Programs/ShaderCompileWorker/Config/DefaultEngine.ini
@@ -1,3 +1,7 @@
-[DerivedDataBackendGraph]
+[DerivedDataCacheGraphs]
+Default=(ProjectPak, EnginePak, ZenLocal, Local, ZenShared, Cloud)
+DerivedDataBackendGraph=(ProjectPak, EnginePak, ZenLocal, Local, ZenShared, Cloud)
+
+[DerivedDataCacheStores]
 Local=(Type=FileSystem, ReadOnly=false, Clean=false, Flush=false, Path=../../../Engine/Programs/%GAME%/DDC)
-Pak=(Type=ReadPak, Filename=../../../Engine/Programs/%GAME%/DDC/DDC.ddp)
+ProjectPak=(Type=ReadPak, Filename=../../../Engine/Programs/%GAME%/DDC/DDC.ddp)
diff --git a/Engine/Programs/ShaderCompileWorker/Config/Mac/MacEngine.ini b/Engine/Programs/ShaderCompileWorker/Config/Mac/MacEngine.ini
--- a/Engine/Programs/ShaderCompileWorker/Config/Mac/MacEngine.ini
+++ b/Engine/Programs/ShaderCompileWorker/Config/Mac/MacEngine.ini
@@ -1,3 +1,2 @@
-[DerivedDataBackendGraph]
+[DerivedDataCacheStores]
 Shared=(Type=FileSystem, ReadOnly=false, Clean=false, Flush=false, DeleteUnused=false, UnusedFileAge=38, FoldersToClean=-1, Path=/Volumes/UE4DDC, EnvPathOverride=UE-SharedDataCachePath)
-AltShared=(Type=FileSystem, ReadOnly=true, Clean=false, Flush=false, DeleteUnused=false, UnusedFileAge=38, FoldersToClean=-1, Path=/Volumes/UE4DDC2, EnvPathOverride=UE-SharedDataCachePath2)
