BaseEngine ini - Pointing default Modular Control Rig Root module to 58 instead of 56 variant.

#Synced-CL 48218623

diff --git a/Engine/Config/BaseEngine.ini b/Engine/Config/BaseEngine.ini
--- a/Engine/Config/BaseEngine.ini
+++ b/Engine/Config/BaseEngine.ini
@@ -3874,7 +3874,7 @@ bShouldEDLPackagesInheritWorldChunkAssignmentsDuringCook=true
 VerificationThreadCount=4
 
 [/Script/ControlRig.ControlRigSettings]
-DefaultRootModule=/ControlRig/Modules/Modules56/Root.Root
+DefaultRootModule=/ControlRig/Modules/Modules58/Root.Root
 
 [/Script/MSGamingSupport.MSGamingSettings]
 bUseSimplifiedUserModel=True
