Add missing IOS devices - ipadpro m5 (8th gen, 11 and 13in models), iphone 17 models, and iphone air

#rnx

diff --git a/Engine/Config/BaseDeviceProfiles.ini b/Engine/Config/BaseDeviceProfiles.ini
--- a/Engine/Config/BaseDeviceProfiles.ini
+++ b/Engine/Config/BaseDeviceProfiles.ini
@@ -86,6 +86,8 @@
 +DeviceProfileNameAndTypes=iPadPro5_129,IOS
 +DeviceProfileNameAndTypes=iPadPro6_129,IOS
 +DeviceProfileNameAndTypes=iPadPro7_129,IOS
++DeviceProfileNameAndTypes=iPadPro8_11,IOS
++DeviceProfileNameAndTypes=iPadPro8_13,IOS
 +DeviceProfileNameAndTypes=AppleTV,IOS
 +DeviceProfileNameAndTypes=AppleTV4K,IOS
 +DeviceProfileNameAndTypes=AppleTV2_4K,IOS
@@ -321,6 +323,8 @@ iPad15,[5,6]=iPadAir7_13
 iPad16,[1,2]=iPadMini7
 iPad16,[3,4]=iPadPro5_11
 iPad16,[5,6]=iPadPro7_129
+iPad17,[1,2]=iPadPro8_11
+iPad17,[3,4]=iPadPro8_13
 AppleTV5,=AppleTV
 AppleTV6,=AppleTV4K
 AppleTV11,=AppleTV2_4K
@@ -797,6 +801,14 @@ BaseProfileName=iPadPro4_11
 DeviceType=IOS
 BaseProfileName=iPadPro4_11
 
+[iPadPro8_11 DeviceProfile]
+DeviceType=IOS
+BaseProfileName=iPadPro4_11
+
+[iPadPro8_13 DeviceProfile]
+DeviceType=IOS
+BaseProfileName=iPadPro4_11
+
 [iPadMini7 DeviceProfile]
 DeviceType=IOS
 BaseProfileName=iPadPro4_11
