Added / updated device profiles for iOS devices: iPhone 17 models, iPhone Air, iPad Pro 8th gen models, iPad 11 (A16).

#rnx

diff --git a/Engine/Config/BaseDeviceProfiles.ini b/Engine/Config/BaseDeviceProfiles.ini
--- a/Engine/Config/BaseDeviceProfiles.ini
+++ b/Engine/Config/BaseDeviceProfiles.ini
@@ -76,6 +76,7 @@
 +DeviceProfileNameAndTypes=iPad8,IOS
 +DeviceProfileNameAndTypes=iPad9,IOS
 +DeviceProfileNameAndTypes=iPad10,IOS
++DeviceProfileNameAndTypes=iPad11,IOS
 +DeviceProfileNameAndTypes=iPadPro11,IOS
 +DeviceProfileNameAndTypes=iPadPro2_11,IOS
 +DeviceProfileNameAndTypes=iPadPro3_11,IOS
@@ -320,6 +321,7 @@ iPad14,[8,9]=iPadAir6
 iPad14,1$=iPadAir_13
 iPad15,[3,4]=iPadAir7
 iPad15,[5,6]=iPadAir7_13
+iPad15,[7,8]=iPad11
 iPad16,[1,2]=iPadMini7
 iPad16,[3,4]=iPadPro5_11
 iPad16,[5,6]=iPadPro7_129
@@ -801,6 +803,10 @@ BaseProfileName=iPadPro4_11
 DeviceType=IOS
 BaseProfileName=iPadPro4_11
 
+[iPad11 DeviceProfile]
+DeviceType=IOS
+BaseProfileName=iPadPro3_11
+
 [iPadPro8_11 DeviceProfile]
 DeviceType=IOS
 BaseProfileName=iPadPro4_11
