Add missing ipad 9th gen and ipadmini 7th gen devices ; fix regex issue with ipad air 6th gen

#rnx
#Synced-CL Unknown

diff --git a/Engine/Config/BaseDeviceProfiles.ini b/Engine/Config/BaseDeviceProfiles.ini
--- a/Engine/Config/BaseDeviceProfiles.ini
+++ b/Engine/Config/BaseDeviceProfiles.ini
@@ -16,7 +16,7 @@
 +DeviceProfileNameAndTypes=iPadAir4,IOS
 +DeviceProfileNameAndTypes=iPadAir5,IOS
 +DeviceProfileNameAndTypes=iPadAir6,IOS
-+DeviceProfileNameAndTypes=iPadAir_13,IOS
++DeviceProfileNameAndTypes=iPadAir6_13,IOS
 +DeviceProfileNameAndTypes=iPadAir7,IOS
 +DeviceProfileNameAndTypes=iPadAir7_13,IOS
 +DeviceProfileNameAndTypes=iPadMini4,IOS
@@ -318,7 +318,7 @@ iPad14,[1,2]=iPadMini6
 iPad14,[3,4]=iPadPro4_11
 iPad14,[5,6]=iPadPro6_129
 iPad14,[8,9]=iPadAir6
-iPad14,1$=iPadAir_13
+iPad14,1[0-1]=iPadAir6_13
 iPad15,[3,4]=iPadAir7
 iPad15,[5,6]=iPadAir7_13
 iPad15,[7,8]=iPad11
@@ -776,7 +776,7 @@ BaseProfileName=iPadPro4_11
 DeviceType=IOS
 BaseProfileName=iPadPro4_11
 
-[iPadAir_13 DeviceProfile]
+[iPadAir6_13 DeviceProfile]
 DeviceType=IOS
 BaseProfileName=iPadPro4_11
 
