Added new install bundle to capture global.utoc/global.uondemandtoc to prevent these files getting included and installed in other bundles.

#rnx

#Synced-CL 48408831

diff --git a/Engine/Config/BaseInstallBundle.ini b/Engine/Config/BaseInstallBundle.ini
--- a/Engine/Config/BaseInstallBundle.ini
+++ b/Engine/Config/BaseInstallBundle.ini
@@ -1,5 +1,10 @@
 ; Defines install bundles and options
 
+; global.utoc/uondemandtoc should never be installed
+[InstallBundleDefinition FilterGlobals]
+Tags=Globals
+FileRegex="global.*(pak|sig|ucas|utoc|uondemandtoc)"
+Order=0
 
 [InstallBundleSource.Bulk.MiscSettings]
 ; Allows our Bulk data source to serialize a BulkBuildData.ini if it was missing during loading.
