Add commandline and environment write-only and read-only overrides for cloud and zenshared DDC.

#Synced-CL 48523383

diff --git a/Engine/Config/BaseEngine.ini b/Engine/Config/BaseEngine.ini
--- a/Engine/Config/BaseEngine.ini
+++ b/Engine/Config/BaseEngine.ini
@@ -2761,7 +2761,7 @@ DerivedDataBackendGraph_Fill_Seattle=(Hierarchical, ProjectPak, EnginePak, ZenLo
 ;   Remote: Whether to force the flags for the node to include the Remote flag. Default is determined by the host URL. Optional.
 ;   ReadOnly: Whether the cache is treated as read-only. Default is false. Optional.
 ZenLocal=(Type=Zen, ServerID=Local)
-ZenShared=(Type=Zen, ServerID=Shared)
+ZenShared=(Type=Zen, ServerID=Shared, EnvReadOnlyOverride=UE-ZenSharedDataCacheReadOnly, CommandLineReadOnlyOverride=ZenSharedDataCacheReadOnly, EnvWriteOnlyOverride=UE-ZenSharedDataCacheWriteOnly, CommandLineWriteOnlyOverride=ZenSharedDataCacheWriteOnly)
 
 ; Type=Cloud is for Unreal Cloud DDC. Available parameters are:
 ;   Host: The URI for the server to connect to, e.g., https://jupiter.example.com.
@@ -2786,7 +2786,7 @@ ZenShared=(Type=Zen, ServerID=Shared)
 ;   HttpVersion: An override for which http version to use. Values: http1.0, http1.1, http2, http2-only.
 ;   EnvHttpVersionOverride: An environment variable that if set will be used instead of HttpVersion.
 ;   CommandLineHttpVersionOverride: A command line argument used in preference to the HttpVersion / EnvHttpVersionOverride settings
-Cloud=(Type=Cloud, ServerID=Cloud, EnvHostOverride=UE-CloudDataCacheHost, CommandLineHostOverride=CloudDataCacheHost, CommandLineOAuthSecretOverride="CloudDataCacheOAuthSecret", EnvOAuthAccessTokenOverride="UE-CloudDataCacheAccessToken", EnvHttpVersionOverride="UE-CloudDataCacheHttpVersion", CommandLineHttpVersionOverride="CloudDataCacheHttpVersion", EnvOAuthProviderIdentifierOverride="UE-CloudDataCacheOAuthProviderIdentifier", CommandLineOAuthProviderIdentifierOverride="CloudDataCacheOAuthProviderIdentifier")
+Cloud=(Type=Cloud, ServerID=Cloud, EnvHostOverride=UE-CloudDataCacheHost, CommandLineHostOverride=CloudDataCacheHost, CommandLineOAuthSecretOverride="CloudDataCacheOAuthSecret", EnvOAuthAccessTokenOverride="UE-CloudDataCacheAccessToken", EnvHttpVersionOverride="UE-CloudDataCacheHttpVersion", CommandLineHttpVersionOverride="CloudDataCacheHttpVersion", EnvOAuthProviderIdentifierOverride="UE-CloudDataCacheOAuthProviderIdentifier", CommandLineOAuthProviderIdentifierOverride="CloudDataCacheOAuthProviderIdentifier", EnvReadOnlyOverride=UE-CloudDataCacheReadOnly, CommandLineReadOnlyOverride=CloudDataCacheReadOnly, EnvWriteOnlyOverride=UE-CloudDataCacheWriteOnly, CommandLineWriteOnlyOverride=CloudDataCacheWriteOnly)
 
 ; Type=FileSystem is for a file system cache. Available parameters are:
 ;	ReadOnly: Can data be written to this layer. Default is false.
