Fix warnings when ACLPlugin is not enabled
- Pass LOAD_NoWarn to LoadObject because GetDefaultAnimationCompressionSettings does custom logging when the requested object fails to load. Without it we get redundant warnings.
- Don't log a warning when bIsFatal is false. We're going to use the fallback value, which is a supported flow. bIsFatal will be true when looking up the fallback value, so we'll log an appropriate message at that time.

#Synced-CL 48571193

diff --git a/Engine/Config/BaseEngine.ini b/Engine/Config/BaseEngine.ini
--- a/Engine/Config/BaseEngine.ini
+++ b/Engine/Config/BaseEngine.ini
@@ -3591,7 +3591,7 @@ BoneTimecodeCustomAttributeNameSettings=(HourAttributeName="TCHour",MinuteAttrib
 ; Fallback is used only if the default codec isn't found (e.g. if the plugin isn't enabled)
 BoneCompressionSettingsFallback="/Engine/Animation/DefaultAnimBoneCompressionSettings"
 ; Default codec assigned when new sequences are imported/created
-BoneCompressionSettings="/ACLPlugin/ACLAnimBoneCompressionSettings"
+BoneCompressionSettings="/Engine/Animation/DefaultAnimBoneCompressionSettings"
 ; Fallback is used only if the default codec isn't found (e.g. if the plugin isn't enabled)
 AnimationRecorderBoneCompressionSettingsFallback="/Engine/Animation/DefaultRecorderBoneCompression"
 ; Default codec for the animation recorder
@@ -3599,7 +3599,7 @@ AnimationRecorderBoneCompressionSettings="/Engine/Animation/DefaultRecorderBoneC
 ; Fallback is used only if the default codec isn't found (e.g. if the plugin isn't enabled)
 CurveCompressionSettingsFallback="/Engine/Animation/DefaultAnimCurveCompressionSettings"
 ; Default codec assigned when new sequences are imported/created
-CurveCompressionSettings="/ACLPlugin/ACLAnimCurveCompressionSettings"
+CurveCompressionSettings="/Engine/Animation/DefaultAnimCurveCompressionSettings"
 VariableFrameStrippingSettings="/Engine/Animation/DefaultVariableFrameStrippingSettings"
 
 [/Script/Engine.MeshSimplificationSettings]
diff --git a/Engine/Plugins/Animation/ACLPlugin/Config/Engine.ini b/Engine/Plugins/Animation/ACLPlugin/Config/Engine.ini
new file mode 100644
--- /dev/null
+++ b/Engine/Plugins/Animation/ACLPlugin/Config/Engine.ini
@@ -0,0 +1,5 @@
+[Animation.DefaultObjectSettings]
+; Default codec assigned when new sequences are imported/created
+BoneCompressionSettings="/ACLPlugin/ACLAnimBoneCompressionSettings"
+; Default codec assigned when new sequences are imported/created
+CurveCompressionSettings="/ACLPlugin/ACLAnimCurveCompressionSettings"
diff --git a/Engine/Source/Runtime/Engine/Private/Animation/AnimationUtils.cpp b/Engine/Source/Runtime/Engine/Private/Animation/AnimationUtils.cpp
--- a/Engine/Source/Runtime/Engine/Private/Animation/AnimationUtils.cpp
+++ b/Engine/Source/Runtime/Engine/Private/Animation/AnimationUtils.cpp
@@ -912,7 +912,7 @@ UObject* GetDefaultAnimationCompressionSettings(const TCHAR* IniValueName, bool
 	}
 
 	const FString& CompressionSettingsName = Value->GetValue();
-	UObject* DefaultCompressionSettings = LoadObject<UObject>(nullptr, *CompressionSettingsName);
+	UObject* DefaultCompressionSettings = LoadObject<UObject>(nullptr, *CompressionSettingsName, {}, LOAD_NoWarn);
 
 	if (DefaultCompressionSettings == nullptr)
 	{
