Allow ZenLocal to have ReadOnly/WriteOnly specified on commandline or env var.

#Synced-CL 48523383

diff --git a/Engine/Config/BaseEngine.ini b/Engine/Config/BaseEngine.ini
--- a/Engine/Config/BaseEngine.ini
+++ b/Engine/Config/BaseEngine.ini
@@ -2761,7 +2761,7 @@ DerivedDataBackendGraph_Fill_Seattle=(Hierarchical, ProjectPak, EnginePak, ZenLo
 ;   Local: Whether to force the flags for the node to include the Local flag. Default is determined by the host URL. Optional.
 ;   Remote: Whether to force the flags for the node to include the Remote flag. Default is determined by the host URL. Optional.
 ;   ReadOnly: Whether the cache is treated as read-only. Default is false. Optional.
-ZenLocal=(Type=Zen, ServerID=Local)
+ZenLocal=(Type=Zen, ServerID=Local, EnvReadOnlyOverride=UE-ZenLocalDataCacheReadOnly, CommandLineReadOnlyOverride=ZenLocalDataCacheReadOnly, EnvWriteOnlyOverride=UE-ZenLocalDataCacheWriteOnly, CommandLineWriteOnlyOverride=ZenLocalDataCacheWriteOnly)
 ZenShared=(Type=Zen, ServerID=Shared, EnvReadOnlyOverride=UE-ZenSharedDataCacheReadOnly, CommandLineReadOnlyOverride=ZenSharedDataCacheReadOnly, EnvWriteOnlyOverride=UE-ZenSharedDataCacheWriteOnly, CommandLineWriteOnlyOverride=ZenSharedDataCacheWriteOnly)
 
 ; Type=Cloud is for Unreal Cloud DDC. Available parameters are:
