Add Unreal Insights settings to the cooked editor build.

#Synced-CL 47929233

diff --git a/Engine/Config/BaseGame.ini b/Engine/Config/BaseGame.ini
--- a/Engine/Config/BaseGame.ini
+++ b/Engine/Config/BaseGame.ini
@@ -271,6 +271,7 @@ SubtitleWidgetToUseDefault=/Script/UMG.WidgetBlueprintGeneratedClass'/SubtitlesA
 +EngineExtraStageFiles=(Path="Config", Files="*Editor*")
 +EngineExtraStageFiles=(Path="Config", Files="*Lightmass*")
 +EngineExtraStageFiles=(Path="Config", Files="*SourceControlSettings*")
++EngineExtraStageFiles=(Path="Config", Files="*UnrealInsights*")
 +EngineExtraStageFiles=(Path="Content/Slate")
 +EngineExtraStageFiles=(Path="Content/Splash")
 +EngineExtraStageFiles=(Path="Content/EngineMaterials")
