Added support for the Buffer Visualization view mode in the mobile post processing path

#virtualized
#Synced-CL 48983662

diff --git a/Engine/Config/BaseEngine.ini b/Engine/Config/BaseEngine.ini
--- a/Engine/Config/BaseEngine.ini
+++ b/Engine/Config/BaseEngine.ini
@@ -3069,30 +3069,30 @@ DefaultUpdateOverlapsMethodDuringLevelStreaming=AlwaysUpdate
 +CollisionChannelRedirects=(OldName="PawnMovement",NewName="Pawn")
 
 [Engine.BufferVisualizationMaterials]
-BaseColor=(Material="/Engine/BufferVisualization/BaseColor.BaseColor", Name=LOCTEXT("BaseColorMat", "Base Color"))
-CustomDepth=(Material="/Engine/BufferVisualization/CustomDepth.CustomDepth", Name=LOCTEXT("BaseCustomDepthMat", "Custom Depth"))
-CustomStencil=(Material="/Engine/BufferVisualization/CustomStencil.CustomStencil", Name=LOCTEXT("BaseCustomStencilMat", "Custom Stencil"))
+BaseColor=(Material="/Engine/BufferVisualization/BaseColor.BaseColor", Name=LOCTEXT("BaseColorMat", "Base Color"), EnableOnMobile=true)
+CustomDepth=(Material="/Engine/BufferVisualization/CustomDepth.CustomDepth", Name=LOCTEXT("BaseCustomDepthMat", "Custom Depth"), EnableOnMobile=true)
+CustomStencil=(Material="/Engine/BufferVisualization/CustomStencil.CustomStencil", Name=LOCTEXT("BaseCustomStencilMat", "Custom Stencil"), EnableOnMobile=true)
 FinalImage=(Material="/Engine/BufferVisualization/FinalImage.FinalImage", Name=LOCTEXT("BaseFinalImageMat", "Final Image"), ApplyAutoExposure=true)
-ShadingModel=(Material="/Engine/BufferVisualization/LightingModel.LightingModel", Name=LOCTEXT("BaseShadingModelMat", "Shading Model"))
-MaterialAO=(Material="/Engine/BufferVisualization/MaterialAO.MaterialAO", Name=LOCTEXT("BaseMaterialAOMat", "Material Ambient Occlusion"))
-Metallic=(Material="/Engine/BufferVisualization/Metallic.Metallic", Name=LOCTEXT("BaseMetallicMat", "Metallic"))
-Opacity=(Material="/Engine/BufferVisualization/Opacity.Opacity", Name=LOCTEXT("BaseOpacityMat", "Opacity"))
-Roughness=(Material="/Engine/BufferVisualization/Roughness.Roughness", Name=LOCTEXT("BaseRoughnessMat", "Roughness"))
-Anisotropy=(Material="/Engine/BufferVisualization/Anisotropy.Anisotropy", Name=LOCTEXT("BaseAnisotropyMat", "Anisotropy"))
-SceneColor=(Material="/Engine/BufferVisualization/SceneColor.SceneColor", Name=LOCTEXT("BaseSceneColorMat", "Scene Color"), ApplyAutoExposure=true)
-SceneDepth=(Material="/Engine/BufferVisualization/SceneDepth.SceneDepth", Name=LOCTEXT("BaseSceneDepthMat", "Scene Depth"))
-SeparateTranslucencyRGB=(Material="/Engine/BufferVisualization/SeparateTranslucencyRGB.SeparateTranslucencyRGB", Name=LOCTEXT("BaseSeparateTranslucencyRGBMat", "Separate Translucency RGB"), ApplyAutoExposure=true)
-SeparateTranslucencyA=(Material="/Engine/BufferVisualization/SeparateTranslucencyA.SeparateTranslucencyA", Name=LOCTEXT("BaseSeparateTranslucencyAMat", "Separate Translucency Alpha"))
-Specular=(Material="/Engine/BufferVisualization/Specular.Specular", Name=LOCTEXT("BaseSpecularMat", "Specular"))
-SubsurfaceColor=(Material="/Engine/BufferVisualization/SubsurfaceColor.SubsurfaceColor", Name=LOCTEXT("BaseSubsurfaceColorMat", "Subsurface Color"))
-WorldNormal=(Material="/Engine/BufferVisualization/WorldNormal.WorldNormal", Name=LOCTEXT("BaseWorldNormalMat", "World Normal"))
-WorldTangent=(Material="/Engine/BufferVisualization/WorldTangent.WorldTangent", Name=LOCTEXT("BaseWorldTangentMat", "World Tangent"))
-AmbientOcclusion=(Material="/Engine/BufferVisualization/AmbientOcclusion.AmbientOcclusion", Name=LOCTEXT("BaseAmbientOcclusionMat", "Ambient Occlusion"))
-CustomDepthWorldUnits=(Material="/Engine/BufferVisualization/CustomDepthWorldUnits.CustomDepthWorldUnits", Name=LOCTEXT("BaseCustomDepthWorldUnitsMat", "Custom Depth World Units"))
-SceneDepthWorldUnits=(Material="/Engine/BufferVisualization/SceneDepthWorldUnits.SceneDepthWorldUnits", Name=LOCTEXT("BaseSceneDepthWorldUnitsMat", "Scene Depth World Units"))
-Velocity=(Material="/Engine/BufferVisualization/Velocity.Velocity", Name=LOCTEXT("Velocity", "Velocity"))
-PreTonemapHDRColor=(Material="/Engine/BufferVisualization/PreTonemapHDRColor.PreTonemapHDRColor", Name=LOCTEXT("PreTonemapHDRColor", "Pre Tonemap HDR Color"), ApplyAutoExposure=true)
-PostTonemapHDRColor=(Material="/Engine/BufferVisualization/PostTonemapHDRColor.PostTonemapHDRColor", Name=LOCTEXT("PostTonemapHDRColor", "Post Tonemap HDR Color"), ApplyAutoExposure=true)
+ShadingModel=(Material="/Engine/BufferVisualization/LightingModel.LightingModel", Name=LOCTEXT("BaseShadingModelMat", "Shading Model"), EnableOnMobile=false)
+MaterialAO=(Material="/Engine/BufferVisualization/MaterialAO.MaterialAO", Name=LOCTEXT("BaseMaterialAOMat", "Material Ambient Occlusion"), EnableOnMobile=true)
+Metallic=(Material="/Engine/BufferVisualization/Metallic.Metallic", Name=LOCTEXT("BaseMetallicMat", "Metallic"), EnableOnMobile=true)
+Opacity=(Material="/Engine/BufferVisualization/Opacity.Opacity", Name=LOCTEXT("BaseOpacityMat", "Opacity"), EnableOnMobile=false)
+Roughness=(Material="/Engine/BufferVisualization/Roughness.Roughness", Name=LOCTEXT("BaseRoughnessMat", "Roughness"), EnableOnMobile=true)
+Anisotropy=(Material="/Engine/BufferVisualization/Anisotropy.Anisotropy", Name=LOCTEXT("BaseAnisotropyMat", "Anisotropy"), EnableOnMobile=false)
+SceneColor=(Material="/Engine/BufferVisualization/SceneColor.SceneColor", Name=LOCTEXT("BaseSceneColorMat", "Scene Color"), ApplyAutoExposure=true, EnableOnMobile=true)
+SceneDepth=(Material="/Engine/BufferVisualization/SceneDepth.SceneDepth", Name=LOCTEXT("BaseSceneDepthMat", "Scene Depth"), EnableOnMobile=true)
+SeparateTranslucencyRGB=(Material="/Engine/BufferVisualization/SeparateTranslucencyRGB.SeparateTranslucencyRGB", Name=LOCTEXT("BaseSeparateTranslucencyRGBMat", "Separate Translucency RGB"), ApplyAutoExposure=true, EnableOnMobile=false)
+SeparateTranslucencyA=(Material="/Engine/BufferVisualization/SeparateTranslucencyA.SeparateTranslucencyA", Name=LOCTEXT("BaseSeparateTranslucencyAMat", "Separate Translucency Alpha"), EnableOnMobile=false)
+Specular=(Material="/Engine/BufferVisualization/Specular.Specular", Name=LOCTEXT("BaseSpecularMat", "Specular"), EnableOnMobile=true)
+SubsurfaceColor=(Material="/Engine/BufferVisualization/SubsurfaceColor.SubsurfaceColor", Name=LOCTEXT("BaseSubsurfaceColorMat", "Subsurface Color"), EnableOnMobile=false)
+WorldNormal=(Material="/Engine/BufferVisualization/WorldNormal.WorldNormal", Name=LOCTEXT("BaseWorldNormalMat", "World Normal"), EnableOnMobile=true)
+WorldTangent=(Material="/Engine/BufferVisualization/WorldTangent.WorldTangent", Name=LOCTEXT("BaseWorldTangentMat", "World Tangent"), EnableOnMobile=false)
+AmbientOcclusion=(Material="/Engine/BufferVisualization/AmbientOcclusion.AmbientOcclusion", Name=LOCTEXT("BaseAmbientOcclusionMat", "Ambient Occlusion"), EnableOnMobile=true)
+CustomDepthWorldUnits=(Material="/Engine/BufferVisualization/CustomDepthWorldUnits.CustomDepthWorldUnits", Name=LOCTEXT("BaseCustomDepthWorldUnitsMat", "Custom Depth World Units"), EnableOnMobile=true)
+SceneDepthWorldUnits=(Material="/Engine/BufferVisualization/SceneDepthWorldUnits.SceneDepthWorldUnits", Name=LOCTEXT("BaseSceneDepthWorldUnitsMat", "Scene Depth World Units"), EnableOnMobile=true)
+Velocity=(Material="/Engine/BufferVisualization/Velocity.Velocity", Name=LOCTEXT("Velocity", "Velocity"), EnableOnMobile=true)
+PreTonemapHDRColor=(Material="/Engine/BufferVisualization/PreTonemapHDRColor.PreTonemapHDRColor", Name=LOCTEXT("PreTonemapHDRColor", "Pre Tonemap HDR Color"), ApplyAutoExposure=true, EnableOnMobile=true)
+PostTonemapHDRColor=(Material="/Engine/BufferVisualization/PostTonemapHDRColor.PostTonemapHDRColor", Name=LOCTEXT("PostTonemapHDRColor", "Post Tonemap HDR Color"), ApplyAutoExposure=true, EnableOnMobile=true)
 
 [DeviceProfileManager]
 DeviceProfileSelectionModule="ExampleDeviceProfileSelector"
diff --git a/Engine/Source/Runtime/Engine/Private/BufferVisualizationData.cpp b/Engine/Source/Runtime/Engine/Private/BufferVisualizationData.cpp
--- a/Engine/Source/Runtime/Engine/Private/BufferVisualizationData.cpp
+++ b/Engine/Source/Runtime/Engine/Private/BufferVisualizationData.cpp
@@ -64,12 +64,16 @@ void FBufferVisualizationData::Initialize()
 							bool bApplyAutoExposure = false;
 							FParse::Bool(*It.Value().GetValue(), TEXT("ApplyAutoExposure="), bApplyAutoExposure);
 
+							bool bEnableOnMobile = true;
+							FParse::Bool(*It.Value().GetValue(), TEXT("EnableOnMobile="), bEnableOnMobile);
+
 							Material->AddToRoot();
 							Record& Rec = MaterialMap.Add(It.Key(), Record());
 							Rec.Name = It.Key().GetPlainNameString();
 							Rec.Material = Material;
 							Rec.bApplyAutoExposure = bApplyAutoExposure;
 							Rec.DisplayName = DisplayName;
+							Rec.bEnableOnMobile = bEnableOnMobile;
 
 							MaterialMapFromMaterialName.Add(Material->GetFName(), Rec);
 						}
@@ -172,6 +176,18 @@ bool FBufferVisualizationData::GetMaterialApplyAutoExposure(FName InMaterialName
 	}
 }
 
+bool FBufferVisualizationData::GetMaterialEnableOnMobile(FName InMaterialName) const
+{
+	if (const Record* Result = GetRecord(InMaterialName))
+	{
+		return Result->bEnableOnMobile;
+	}
+	else
+	{
+		return false;
+	}
+}
+
 FText FBufferVisualizationData::GetMaterialDisplayName(FName InMaterialName) const
 {
 	if (const Record* Result = GetRecord(InMaterialName))
diff --git a/Engine/Source/Runtime/Engine/Public/BufferVisualizationData.h b/Engine/Source/Runtime/Engine/Public/BufferVisualizationData.h
--- a/Engine/Source/Runtime/Engine/Public/BufferVisualizationData.h
+++ b/Engine/Source/Runtime/Engine/Public/BufferVisualizationData.h
@@ -26,6 +26,9 @@ public:
 	/** Get whether auto exposure should be applied for the visualization mode **/
 	ENGINE_API bool GetMaterialApplyAutoExposure(FName InMaterialName) const;
 
+	/** Get whether this visualization mode should be enabled on mobile **/
+	ENGINE_API bool GetMaterialEnableOnMobile(FName InMaterialName) const;
+
 	/** Get the display name of a named material from the available material map **/
 	ENGINE_API FText GetMaterialDisplayName(FName InMaterialName) const;
 
@@ -64,6 +67,7 @@ private:
 		FText DisplayName;
 		UMaterialInterface* Material;
 		bool bApplyAutoExposure;
+		bool bEnableOnMobile;
 	};
 
 	/** Mapping of FName (first parameter in ini file material list) to a material record */
diff --git a/Engine/Source/Runtime/Renderer/Private/PostProcess/PostProcessVisualizeBuffer.cpp b/Engine/Source/Runtime/Renderer/Private/PostProcess/PostProcessVisualizeBuffer.cpp
--- a/Engine/Source/Runtime/Renderer/Private/PostProcess/PostProcessVisualizeBuffer.cpp
+++ b/Engine/Source/Runtime/Renderer/Private/PostProcess/PostProcessVisualizeBuffer.cpp
@@ -14,6 +14,7 @@
 #include "UnrealEngine.h"
 #include "PathTracing.h"
 #include "Materials/Material.h"
+#include "Internationalization/Internationalization.h"
 
 class FVisualizeBufferPS : public FGlobalShader
 {
@@ -538,10 +539,25 @@ FScreenPassTexture AddVisualizeGBufferPass(
 				DownsampleInputs.SceneColor = FScreenPassTextureSlice::CreateFromScreenPassTexture(GraphBuilder, HalfSize);
 
 				bool bMaterialRunsBeforeTonemapping = RunsBeforeTonemapping(MaterialInterface);
+				const bool bIsMobilePlatform = IsMobilePlatform(View.GetShaderPlatform());
 
 				FVisualizeBufferTile Tile;
 				Tile.Input = AddDownsamplePass(GraphBuilder, View, DownsampleInputs);
-				Tile.Label = GetBufferVisualizationData().GetMaterialDisplayName(FName(*MaterialName)).ToString();
+
+				const FName MaterialFName = MaterialInterface->GetFName();
+				if (!bIsMobilePlatform || GetBufferVisualizationData().GetMaterialEnableOnMobile(MaterialFName))
+				{
+					Tile.Label = GetBufferVisualizationData().GetMaterialDisplayName(MaterialFName).ToString();
+				}
+				else
+				{
+					const FText NotImplementedOnMobile = NSLOCTEXT("PostProcessVisualizeBuffer", "NotImplementedOnMobile", "(Not implemented on Mobile)");
+					Tile.Label = FText::Format(
+						NSLOCTEXT("PostProcessVisualizeBuffer", "TileNotImplementedOnMobile", "{0} {1}"),
+						GetBufferVisualizationData().GetMaterialDisplayName(MaterialFName),
+						NotImplementedOnMobile).ToString();
+				}
+
 				// VisualizeGBufferOverviewPass runs after tonemapping, so any materials that were expecting to run before
 				// need to be tonemapped manually.
 				Tile.bNeedsGammaCorrection = bMaterialRunsBeforeTonemapping;
diff --git a/Engine/Source/Runtime/Renderer/Private/PostProcess/PostProcessing.cpp b/Engine/Source/Runtime/Renderer/Private/PostProcess/PostProcessing.cpp
--- a/Engine/Source/Runtime/Renderer/Private/PostProcess/PostProcessing.cpp
+++ b/Engine/Source/Runtime/Renderer/Private/PostProcess/PostProcessing.cpp
@@ -2573,6 +2573,8 @@ void AddMobilePostProcessingPasses(FRDGBuilder& GraphBuilder, FScene* Scene, con
 	const FEyeAdaptationParameters EyeAdaptationParameters = GetEyeAdaptationParameters(View);
 	FRDGBufferRef LastEyeAdaptationBuffer = GetEyeAdaptationBuffer(GraphBuilder, View);
 
+	FVisualizeBufferContext VisualizeBufferContext{};
+
 	enum class EPass : uint32
 	{
 		Distortion,
@@ -2598,6 +2600,7 @@ void AddMobilePostProcessingPasses(FRDGBuilder& GraphBuilder, FScene* Scene, con
 		SecondaryUpscale,
 		Visualize,
 		VisualizeLightGrid,
+		VisualizeGBuffer,
 		HMDDistortion,
 		MAX
 	};
@@ -2644,6 +2647,7 @@ void AddMobilePostProcessingPasses(FRDGBuilder& GraphBuilder, FScene* Scene, con
 		TEXT("SecondaryUpscale"),
 		TEXT("Visualize"),
 		TEXT("VisualizeLightGrid"),
+		TEXT("VisualizeGBuffer"),
 		TEXT("HMDDistortion")
 	};
 
@@ -2687,6 +2691,9 @@ void AddMobilePostProcessingPasses(FRDGBuilder& GraphBuilder, FScene* Scene, con
 	//Primary does not require the interface to be available to run as it has a fallback logic already.
 	bool bShouldPrimaryUpscale = (View.PrimaryScreenPercentageMethod == EPrimaryScreenPercentageMethod::SpatialUpscale && View.GetSecondaryViewRectSize() != View.ViewRect.Size()) || View.LensDistortionLUT.IsEnabled();
 
+	const bool bVisualizeGBuffer = IsVisualizeGBufferEnabled(View);
+	const bool bOutputInHDR = IsPostProcessingOutputInHDR();
+
 	PassSequence.SetEnabled(EPass::Tonemap, bUseToneMapper);
 	PassSequence.SetEnabled(EPass::HighResolutionScreenshotMask, bUseHighResolutionScreenshotMask);
 #if WITH_EDITOR
@@ -2710,6 +2717,7 @@ void AddMobilePostProcessingPasses(FRDGBuilder& GraphBuilder, FScene* Scene, con
 
 	PassSequence.SetEnabled(EPass::Visualize, View.Family->EngineShowFlags.ShaderComplexity);
 	PassSequence.SetEnabled(EPass::VisualizeLightGrid, ShouldVisualizeLightGrid(View.Family->GetShaderPlatform()));
+	PassSequence.SetEnabled(EPass::VisualizeGBuffer, bVisualizeGBuffer);
 	PassSequence.SetEnabled(EPass::HMDDistortion, View.Family->EngineShowFlags.StereoRendering && View.Family->EngineShowFlags.HMDDistortion);
 
 	// TODO: Add support for EPass::AlphaInvert if the AlphaInvert show flag is enabled, but EPass::PrimaryUpscale is disabled
@@ -2862,7 +2870,7 @@ void AddMobilePostProcessingPasses(FRDGBuilder& GraphBuilder, FScene* Scene, con
 		}
 
 		PassSequence.Finalize();
-			
+
 		if (PassSequence.IsEnabled(EPass::Distortion))
 		{
 			PassSequence.AcceptPass(EPass::Distortion);
@@ -2887,6 +2895,11 @@ void AddMobilePostProcessingPasses(FRDGBuilder& GraphBuilder, FScene* Scene, con
 
 		AddPostProcessMaterialPass(BL_SceneColorBeforeDOF);
 
+		{
+			FPostProcessMaterialInputs PostProcessMaterialInputs = GetPostProcessMaterialInputs(FScreenPassTextureSlice::CreateFromScreenPassTexture(GraphBuilder, SceneColor));
+			AddVisualizeBufferMaterialPass(GraphBuilder, View, BL_SceneColorBeforeDOF, PostProcessMaterialInputs, bOutputInHDR, VisualizeBufferContext);
+		}
+
 		// Optional fixed pass processes
 		if (PassSequence.IsEnabled(EPass::SunMask))
 		{
@@ -3207,6 +3220,11 @@ void AddMobilePostProcessingPasses(FRDGBuilder& GraphBuilder, FScene* Scene, con
 
 		AddPostProcessMaterialPass(BL_SceneColorAfterDOF);
 
+		{
+			FPostProcessMaterialInputs PostProcessMaterialInputs = GetPostProcessMaterialInputs(FScreenPassTextureSlice::CreateFromScreenPassTexture(GraphBuilder, SceneColor));
+			AddVisualizeBufferMaterialPass(GraphBuilder, View, BL_SceneColorAfterDOF, PostProcessMaterialInputs, bOutputInHDR, VisualizeBufferContext);
+		}
+
 		// Temporal Anti-aliasing. Also may perform a temporal upsample from primary to secondary view rect.
 		if (PassSequence.IsEnabled(EPass::TAA))
 		{
@@ -3279,6 +3297,11 @@ void AddMobilePostProcessingPasses(FRDGBuilder& GraphBuilder, FScene* Scene, con
 
 	AddPostProcessMaterialPass(BL_SceneColorBeforeBloom);
 
+	{
+		FPostProcessMaterialInputs PostProcessMaterialInputs = GetPostProcessMaterialInputs(FScreenPassTextureSlice::CreateFromScreenPassTexture(GraphBuilder, SceneColor));
+		AddVisualizeBufferMaterialPass(GraphBuilder, View, BL_SceneColorBeforeBloom, PostProcessMaterialInputs, bOutputInHDR, VisualizeBufferContext);
+	}
+
 	FScreenPassTextureSlice SceneColorBeforeTonemapSlice = FScreenPassTextureSlice::CreateFromScreenPassTexture(GraphBuilder, SceneColor);
 	
 	if (PassSequence.IsEnabled(EPass::Tonemap))
@@ -3389,6 +3412,12 @@ void AddMobilePostProcessingPasses(FRDGBuilder& GraphBuilder, FScene* Scene, con
 		{
 			AddPostProcessMaterialPass(BL_SceneColorAfterTonemapping);
 		}
+
+		{
+			FPostProcessMaterialInputs PostProcessMaterialInputs = GetPostProcessMaterialInputs(FScreenPassTextureSlice::CreateFromScreenPassTexture(GraphBuilder, SceneColor));
+			AddVisualizeBufferMaterialPass(GraphBuilder, View, BL_SceneColorAfterTonemapping, PostProcessMaterialInputs, bOutputInHDR, VisualizeBufferContext);
+		}
+
 	}
 
 	if (PassSequence.IsEnabled(EPass::HighResolutionScreenshotMask))
@@ -3572,6 +3601,16 @@ void AddMobilePostProcessingPasses(FRDGBuilder& GraphBuilder, FScene* Scene, con
 		}
 	}
 
+	if (PassSequence.IsEnabled(EPass::VisualizeGBuffer))
+	{
+		FVisualizeGBufferInputs PassInputs;
+		PassSequence.AcceptOverrideIfLastPass(EPass::VisualizeGBuffer, PassInputs.OverrideOutput);
+		PassInputs.SceneColor = SceneColor;
+		PassInputs.bOutputInHDR = bOutputInHDR;
+
+		SceneColor = AddVisualizeGBufferPass(GraphBuilder, View, PassInputs, VisualizeBufferContext);
+	}
+
 	if (PassSequence.IsEnabled(EPass::VisualizeLightGrid))
 	{
 		FScreenPassRenderTarget OverrideOutput;
