Skip precaching slate materials and only precache particles at proxy time for PC shader profile with 6 or less cores by default to reduce the amount of unused PSOs

#Synced-CL 49066199

diff --git a/Engine/Config/BaseEngine.ini b/Engine/Config/BaseEngine.ini
--- a/Engine/Config/BaseEngine.ini
+++ b/Engine/Config/BaseEngine.ini
@@ -3896,3 +3896,11 @@ DefaultRootModule=/ControlRig/Modules/Modules58/Root.Root
 [/Script/MSGamingSupport.MSGamingSettings]
 bUseSimplifiedUserModel=True
 
+[/Script/Engine.PSOPrecacheSettingsManager]
+!Settings=ClearArray
+; Default settings for max thread count
++Settings=(ShaderPlatform="PCD3D_SM6", MaxCPUThreadCount=255)
++Settings=(ShaderPlatform="PCD3D_ES3_1", MaxCPUThreadCount=255)
+; Don't precache slate and only precache particles at proxy creation time when low number of cores
++Settings=(ShaderPlatform="PCD3D_SM6", MaxCPUThreadCount=6, bSlateMaterials=False, ParticlePrecachingTime=1)
++Settings=(ShaderPlatform="PCD3D_ES3_1", MaxCPUThreadCount=6, bSlateMaterials=False, ParticlePrecachingTime=1)
\ No newline at end of file
