Disable New Sound Wave Player Operator in BaseEngine.ini

#rnx
#Synced-CL 48881229

diff --git a/Engine/Config/BaseEngine.ini b/Engine/Config/BaseEngine.ini
--- a/Engine/Config/BaseEngine.ini
+++ b/Engine/Config/BaseEngine.ini
@@ -1642,6 +1642,9 @@ ConnectionHelpUrl=""
 UnattendedRetryCount=0
 UnattendedRetryTimer=0
 
+[ConsoleVariables]
+au.MetaSound.WavePlayer.UseNewOperator=0
+
 [/Script/Virtualization.VirtualizationFilterSettings]
 !ExcludePackagePaths=ClearArray
 
