Re-enable new SoundWavePlayer operator

#rnx
#Synced-CL 49063521

diff --git a/Engine/Config/BaseEngine.ini b/Engine/Config/BaseEngine.ini
--- a/Engine/Config/BaseEngine.ini
+++ b/Engine/Config/BaseEngine.ini
@@ -1642,9 +1642,6 @@ ConnectionHelpUrl=""
 UnattendedRetryCount=0
 UnattendedRetryTimer=0
 
-[ConsoleVariables]
-au.MetaSound.WavePlayer.UseNewOperator=0
-
 [/Script/Virtualization.VirtualizationFilterSettings]
 !ExcludePackagePaths=ClearArray
 
