Add default PSO precache settings for Android and iOS

#Synced-CL 49265942

diff --git a/Engine/Config/BaseEngine.ini b/Engine/Config/BaseEngine.ini
--- a/Engine/Config/BaseEngine.ini
+++ b/Engine/Config/BaseEngine.ini
@@ -3903,4 +3903,8 @@ bUseSimplifiedUserModel=True
 +Settings=(ShaderPlatform="PCD3D_ES3_1", MaxCPUThreadCount=255)
 ; Don't precache slate and only precache particles at proxy creation time when low number of cores
 +Settings=(ShaderPlatform="PCD3D_SM6", MaxCPUThreadCount=6, bSlateMaterials=False, ParticlePrecachingTime=1)
-+Settings=(ShaderPlatform="PCD3D_ES3_1", MaxCPUThreadCount=6, bSlateMaterials=False, ParticlePrecachingTime=1)
\ No newline at end of file
++Settings=(ShaderPlatform="PCD3D_ES3_1", MaxCPUThreadCount=6, bSlateMaterials=False, ParticlePrecachingTime=1)
+; Don't precache global shaders, slate and only precache particles at proxy creation time on ES31 Android and iOS
++Settings=(ShaderPlatform="SP_OPENGL_ES3_1_ANDROID", GlobalShaderMode=0, bSlateMaterials=False, ParticlePrecachingTime=1)
++Settings=(ShaderPlatform="SP_VULKAN_ES3_1_ANDROID", GlobalShaderMode=0, bSlateMaterials=False, ParticlePrecachingTime=1)
++Settings=(ShaderPlatform="SP_METAL_ES3_1_IOS", GlobalShaderMode=0, bSlateMaterials=False, ParticlePrecachingTime=1)
\ No newline at end of file
