Less aggressive culling of small objects from Lumen Scene on Epic and Cinematic GI quality. Reduces pops in emissive lighting and black spots in reflections from small objects.
* Reduced r.LumenScene.SurfaceCache.CardMinResolution from 4 to 2

#Synced-CL 49613554

diff --git a/Engine/Config/BaseScalability.ini b/Engine/Config/BaseScalability.ini
--- a/Engine/Config/BaseScalability.ini
+++ b/Engine/Config/BaseScalability.ini
@@ -362,7 +362,7 @@ r.Lumen.DiffuseIndirect.Allow=1
 r.Lumen.FinalGatherMethod=1
 r.LumenScene.SurfaceCache.AtlasSize=4096
 r.LumenScene.SurfaceCache.CardTexelDensityScale=100
-r.LumenScene.SurfaceCache.CardMinResolution=4
+r.LumenScene.SurfaceCache.CardMinResolution=2
 r.LumenScene.SurfaceCache.CardMaxResolution=512
 r.LumenScene.DirectLighting.MaxLightsPerTile=8
 r.LumenScene.DirectLighting.UpdateFactor=32
@@ -407,7 +407,7 @@ r.Lumen.DiffuseIndirect.Allow=1
 r.Lumen.FinalGatherMethod=1
 r.LumenScene.SurfaceCache.AtlasSize=4096
 r.LumenScene.SurfaceCache.CardTexelDensityScale=100
-r.LumenScene.SurfaceCache.CardMinResolution=4
+r.LumenScene.SurfaceCache.CardMinResolution=2
 r.LumenScene.SurfaceCache.CardMaxResolution=512
 r.LumenScene.DirectLighting.MaxLightsPerTile=8
 r.LumenScene.DirectLighting.UpdateFactor=32
