Fixed typo in BaseEditorSettings.ini

[Fixed] DisplayPluginFolder typo. Should be DisplayPluginFolders so that the base settings gets applied.

#Synced-CL 49621724

diff --git a/Engine/Config/BaseEditorSettings.ini b/Engine/Config/BaseEditorSettings.ini
--- a/Engine/Config/BaseEditorSettings.ini
+++ b/Engine/Config/BaseEditorSettings.ini
@@ -31,7 +31,7 @@ DisplayCollections=False
 ; Whether to display the C++ folders in the Content Browser
 DisplayCppFolders=True
 ; Whether to display the game project plugins in the Content Browser
-DisplayPluginFolder=True
+DisplayPluginFolders=True
 ; The number of objects to keep in the Content Browser Recently Opened filter 
 NumObjectsInRecentList = 20
 ; Whether to consider an asset's class name when searching by text
