Adding base device profiles for Steam Frame

#Synced-CL 49959945
#Synced-CL 49959945

diff --git a/Engine/Config/BaseDeviceProfiles.ini b/Engine/Config/BaseDeviceProfiles.ini
--- a/Engine/Config/BaseDeviceProfiles.ini
+++ b/Engine/Config/BaseDeviceProfiles.ini
@@ -177,6 +177,7 @@
 +DeviceProfileNameAndTypes=Oculus_Quest,Android
 +DeviceProfileNameAndTypes=Android_OpenXR,Android
 +DeviceProfileNameAndTypes=MagicLeap_Vulkan,Android
++DeviceProfileNameAndTypes=Valve_Frame_Lepton,Android
 
 [GlobalDefaults DeviceProfile]
 TextureLODGroups=(Group=TEXTUREGROUP_World,MinLODSize=1,MaxLODSize=16384,LODBias=0,MinMagFilter=aniso,MipFilter=point,MipGenSettings=TMGS_SimpleAverage)
@@ -866,6 +867,7 @@ BaseProfileName=VisionPro
 +MatchProfile=(Profile="Oculus_Quest2",Match=((SourceType=SRC_DeviceMake,CompareType=CMP_Equal,MatchString="Oculus"), (SourceType=SRC_HMDSystemName,CompareType=CMP_Equal,MatchString="Oculus Quest2")))
 +MatchProfile=(Profile="Oculus_Quest",Match=((SourceType=SRC_DeviceMake,CompareType=CMP_Equal,MatchString="Oculus")))
 +MatchProfile=(Profile="MagicLeap_Vulkan",Match=((SourceType=SRC_DeviceMake,CompareType=CMP_Equal,MatchString="Magic Leap"), (SourceType=SRC_VulkanAvailable,CompareType=CMP_Equal,MatchString="true")))
++MatchProfile=(Profile="Valve_Frame_Lepton",Match=((SourceType=SRC_DeviceMake,CompareType=CMP_Equal,MatchString="Valve"),(SourceType=SRC_DeviceModel,CompareType=CMP_Equal,MatchString="Lepton"),(SourceType=SRC_GpuFamily,CompareType=CMP_Regex,MatchString=".*Adreno \\(TM\\) 7[0-9][0-9].*"),(SourceType=SRC_AndroidVersion, CompareType=CMP_Regex,MatchString="([0-9]+).*"),(SourceType=SRC_PreviousRegexMatch,CompareType=CMP_GreaterEqual,MatchString="9"),(SourceType=SRC_VulkanAvailable,CompareType=CMP_Equal,MatchString="true")))
 
 ; Match Android Emulator
 +MatchProfile=(Profile="Android_Emulator",Match=((SourceType=SRC_DeviceMake,CompareType=CMP_Equal,MatchString="Google"),(SourceType=SRC_GpuFamily,CompareType=CMP_Regex,MatchString="Android Emulator")))
@@ -1512,6 +1514,21 @@ BaseProfileName=Android_Mid
 +CVars=r.Vulkan.SupportsBCTextureFormats=1
 +CVars=r.Vulkan.MaxBarriersPerBatch=128
 
+; cv is Steam Frame
+[cv DeviceProfile]
+DeviceType=LinuxArm64
+BaseProfileName=LinuxArm64
++CVars=r.AllowOcclusionQueries=False
++CVars=r.DefaultBackBufferPixelFormat=0
++CVars=fx.Niagara.QualityLevel.Min=0
++CVars=fx.Niagara.QualityLevel.Max=3
++CVars=r.Mobile.UseHWsRGBEncoding=1
+
+[Valve_Frame_Lepton DeviceProfile]
+DeviceType=Android
+BaseProfileName=Meta_Quest_3
++CVars=xr.DisableOpenXROnAndroidWithoutOculus=0
+
 [Windows_Preview_ES31 DeviceProfile]
 DeviceType=Windows
 BaseProfileName=Windows
