- Fix android shader platform names so correct default PSO precache settings are picked up
- Make sure settings are picked up as well when not defining max CPU cores

#Synced-CL 50811283

diff --git a/Engine/Config/BaseEngine.ini b/Engine/Config/BaseEngine.ini
--- a/Engine/Config/BaseEngine.ini
+++ b/Engine/Config/BaseEngine.ini
@@ -3925,6 +3925,6 @@ bUseSimplifiedUserModel=True
 +Settings=(ShaderPlatform="PCD3D_SM6", MaxCPUThreadCount=6, bSlateMaterials=False, ParticlePrecachingTime=1)
 +Settings=(ShaderPlatform="PCD3D_ES3_1", MaxCPUThreadCount=6, bSlateMaterials=False, ParticlePrecachingTime=1)
 ; Don't precache global shaders, slate and only precache particles at proxy creation time on ES31 Android and iOS
-+Settings=(ShaderPlatform="SP_OPENGL_ES3_1_ANDROID", GlobalShaderMode=0, bSlateMaterials=False, ParticlePrecachingTime=1)
-+Settings=(ShaderPlatform="SP_VULKAN_ES3_1_ANDROID", GlobalShaderMode=0, bSlateMaterials=False, ParticlePrecachingTime=1)
-+Settings=(ShaderPlatform="SP_METAL_ES3_1_IOS", GlobalShaderMode=0, bSlateMaterials=False, ParticlePrecachingTime=1)
\ No newline at end of file
++Settings=(ShaderPlatform="OPENGL_ES3_1_ANDROID", GlobalShaderMode=0, bSlateMaterials=False, ParticlePrecachingTime=1)
++Settings=(ShaderPlatform="VULKAN_ES3_1_ANDROID", GlobalShaderMode=0, bSlateMaterials=False, ParticlePrecachingTime=1)
++Settings=(ShaderPlatform="METAL_ES3_1_IOS", GlobalShaderMode=0, bSlateMaterials=False, ParticlePrecachingTime=1)
\ No newline at end of file
diff --git a/Engine/Source/Runtime/Engine/Private/PSOPrecacheSettings.cpp b/Engine/Source/Runtime/Engine/Private/PSOPrecacheSettings.cpp
--- a/Engine/Source/Runtime/Engine/Private/PSOPrecacheSettings.cpp
+++ b/Engine/Source/Runtime/Engine/Private/PSOPrecacheSettings.cpp
@@ -246,7 +246,7 @@ const FPSOPrecacheSettings& UPSOPrecacheSettingsManager::GetSettings()
 		int32 NumThreads = FPlatformMisc::NumberOfCoresIncludingHyperthreads();
 		for (const FPSOPrecacheSettings& IterSettings : Manager.Settings)
 		{
-			if (IterSettings.ShaderPlatform == ShaderPlatformName && NumThreads <= IterSettings.MaxCPUThreadCount && IterSettings.MaxCPUThreadCount < Settings.MaxCPUThreadCount)
+			if (IterSettings.ShaderPlatform == ShaderPlatformName && NumThreads <= IterSettings.MaxCPUThreadCount && IterSettings.MaxCPUThreadCount <= Settings.MaxCPUThreadCount)
 			{
 				Settings = IterSettings;
 			}
