AssetRegistry: Fix overflow in TagDataMaps' AnsiStrings when initializing the generated AssetRegistry at beginning of cook, on projects with many actors. The Actor metadata strings can accumulate to multiple gigabytes. Remove them from the cooked runtime and cooked development AssetRegistries, which use an FStore and therefore attempt to store accumulated tag strings in a single array of bytes. Also remove the existing FiBData, used by blueprints, from all types, because types such as UWorld that have blueprints internally can have that tag as well.
#rnx
#Synced-CL 51086784

diff --git a/Engine/Config/BaseEngine.ini b/Engine/Config/BaseEngine.ini
--- a/Engine/Config/BaseEngine.ini
+++ b/Engine/Config/BaseEngine.ini
@@ -2873,7 +2873,11 @@ AspectRatioAxisConstraint=AspectRatio_MaintainYFOV
 ; You can use a * wildcard to refer to either all classes or all tags for the class.
 ; If your game does not need any asset registry tags, simply add (Class=*,Tag=*). Be warned that this this may make some engine systems fail to work since they might need ParentClass or GeneratedClass on Blueprints.
 +CookedTagsBlacklist=(Class=/Script/Engine.Blueprint,Tag=FiB)
-+CookedTagsBlacklist=(Class=/Script/Engine.Blueprint,Tag=FiBData)
++CookedTagsBlacklist=(Class=*,Tag=FiBData)
++CookedTagsBlacklist=(Class=*,Tag=ActorMetaData)
++CookedTagsBlacklist=(Class=*,Tag=ActorsMetaData)
++CookedTagsBlacklist=(Class=*,Tag=ActorMetaDataClass)
++CookedTagsBlacklist=(Class=*,Tag=ActorLabel)
 +CookedTagsBlacklist=(Class=*,Tag=AssetImportData)
 
 ; Alternatively, you can use an allow list which will exclude all tags except those mentioned in this list. Set bUseAssetRegistryTagsWhitelistInsteadOfBlacklist to true to use this list.
