Iris - Fixed typo in BaseEngine.ini causing PlayerStates not be prioritized as expected unless a project overrides.

#Synced-CL 51087404

diff --git a/Engine/Config/BaseEngine.ini b/Engine/Config/BaseEngine.ini
--- a/Engine/Config/BaseEngine.ini
+++ b/Engine/Config/BaseEngine.ini
@@ -1477,7 +1477,7 @@ DefaultSpatialFilterName=Spatial
 
 ; PrioritizerConfigs
 ; PlayerState is by default bAlwaysRelevant but the default player state prioritizer handles this.
-+PrioritizerConfigs=(ClassName=/Script/Engine.PlayerState, PrioritizerName=PlayerState, bForceEnableOnAllInstances=true)
++PrioritizerConfigs=(ClassName=/Script/Engine.PlayerState, PrioritizerName=PlayerStatePrioritizer, bForceEnableOnAllInstances=true)
 
 ; DeltaCompressionConfigs
 ; Classes for which deltacompression is enabled
