Change an outdated and misleading comment about the PSO cache.

#Synced-CL 51373877

diff --git a/Engine/Config/BaseEngine.ini b/Engine/Config/BaseEngine.ini
--- a/Engine/Config/BaseEngine.ini
+++ b/Engine/Config/BaseEngine.ini
@@ -2174,7 +2174,7 @@ bUseVirtualCores = False
 ; Note: CookerMemoryUsedInGB, MemoryToLeaveForTheOSInGB, MemoryUsedPerSCWProcessInGB must all be set to enable
 bForceUseSCWMemoryPressureLimits = False
 
-; Temporarily disable chunking of the PSO cache as the runtime code isn't yet in place
+; Disable chunking of the PSO cache by default (cold path).
 bDoNotChunkPSOCache=True
 
 [ShaderCompiler]
