Update GPU driver checks to be more accurate and recent

#Synced-CL 51452150

diff --git a/Engine/Config/BaseHardware.ini b/Engine/Config/BaseHardware.ini
--- a/Engine/Config/BaseHardware.ini
+++ b/Engine/Config/BaseHardware.ini
@@ -36,8 +36,8 @@
 [GPU_NVIDIA Windows]
 SuggestedDriverVersion="560.81"
 
-; The newest driver which supports legacy Nvidia GPUs (700 series and older) is 474.64. We allow this driver for DX11.
-+DriverDenyList=(DriverVersion="<474.64", Reason="These driver versions have known stability issues and missing features", RHI="D3D11")
+; The newest driver which supports legacy Nvidia GPUs (700 series and older) is 475.14. We allow this driver for DX11.
++DriverDenyList=(DriverVersion="<475.14", Reason="These driver versions have known stability issues and missing features", RHI="D3D11")
 
 ; DX12 require recent drivers. The 536 series was released in June 2023.
 +DriverDenyList=(DriverVersion="<536.00", Reason="These driver versions have known stability issues and missing features", RHI="D3D12")
@@ -59,7 +59,7 @@ SuggestedDriverVersion="570.86"
 SuggestedDriverVersion="24.3.1"
 
 ; The newest driver which supports legacy AMD GPUs (GCN1-GCN3 chips like the Radeon R9 or older) is 22.6.1. We allow this driver for DX11.
-+DriverDenyList=(DriverVersion="<27.20.20913.2000", Reason="These driver versions have known stability issues and missing features", RHI="D3D11")
++DriverDenyList=(DriverVersion="<27.20.20912.1002", Reason="These driver versions have known stability issues and missing features", RHI="D3D11")
 
 ; The minimum version for DX12 is 24.3.1, published in March 2024. There are separate packages for different GPU generations and they
 ; have different internal versions, so we can't check for a version number explicitly (the RDNA driver is 31.0.24027.1012, older GPUs have
