Decouple project setting from scalability-based setting to enable or disable Heterogeneous Volumes.

r.HeterogeneousVolumes.Allow (default=1)
This option allows scalability controllers to modify the use of Heterogeneous Volumes, in conjunction with the Project Setting value

EffectsQuality scalability configurations drives the Allow value based on these general rules:
Allow is 1 for High, Epic, and Cinematic scalability. Otherwise, Allow is set to 0.

#Synced-CL 51500622

diff --git a/Engine/Config/BaseScalability.ini b/Engine/Config/BaseScalability.ini
--- a/Engine/Config/BaseScalability.ini
+++ b/Engine/Config/BaseScalability.ini
@@ -749,7 +749,7 @@ r.SkyAtmosphere.TransmittanceLUT.UseSmallFormat=1
 r.SkyAtmosphere.TransmittanceLUT.SampleCount=10.0
 r.SkyAtmosphere.MultiScatteringLUT.SampleCount=15.0
 fx.Niagara.QualityLevel=0
-r.HeterogeneousVolumes=0
+r.HeterogeneousVolumes.Allowed=0
 
 [EffectsQuality@1]
 r.TranslucencyLightingVolume.Dim=32
@@ -777,7 +777,7 @@ r.SkyAtmosphere.TransmittanceLUT.UseSmallFormat=0
 r.SkyAtmosphere.TransmittanceLUT.SampleCount=10.0
 r.SkyAtmosphere.MultiScatteringLUT.SampleCount=15.0
 fx.Niagara.QualityLevel=1
-r.HeterogeneousVolumes=0
+r.HeterogeneousVolumes.Allowed=0
 
 [EffectsQuality@2]
 r.TranslucencyLightingVolume.Dim=48
@@ -805,6 +805,7 @@ r.SkyAtmosphere.TransmittanceLUT.UseSmallFormat=0
 r.SkyAtmosphere.TransmittanceLUT.SampleCount=10.0
 r.SkyAtmosphere.MultiScatteringLUT.SampleCount=15.0
 fx.Niagara.QualityLevel=2
+r.HeterogeneousVolumes.Allowed=1
 r.HeterogeneousVolumes.DownsampleFactor=2
 r.HeterogeneousVolumes.MaxStepCount=96
 r.HeterogeneousVolumes.Shadows.Resolution=256
@@ -836,6 +837,7 @@ r.SkyAtmosphere.TransmittanceLUT.SampleCount=10.0
 r.SkyAtmosphere.MultiScatteringLUT.SampleCount=15.0
 fx.Niagara.QualityLevel=3
 r.Refraction.OffsetQuality=1
+r.HeterogeneousVolumes.Allowed=1
 r.HeterogeneousVolumes.DownsampleFactor=2
 r.HeterogeneousVolumes.MaxStepCount=256
 r.HeterogeneousVolumes.Shadows.Resolution=512
@@ -869,6 +871,7 @@ r.SkyAtmosphere.TransmittanceLUT.SampleCount=30.0
 r.SkyAtmosphere.MultiScatteringLUT.SampleCount=20.0
 fx.Niagara.QualityLevel=4
 r.Refraction.OffsetQuality=1
+r.HeterogeneousVolumes.Allowed=1
 r.HeterogeneousVolumes.ApplyFogInscattering=1
 r.HeterogeneousVolumes.DownsampleFactor=1
 r.HeterogeneousVolumes.MaxStepCount=1024
diff --git a/Engine/Source/Runtime/Renderer/Private/HeterogeneousVolumes/HeterogeneousVolumes.cpp b/Engine/Source/Runtime/Renderer/Private/HeterogeneousVolumes/HeterogeneousVolumes.cpp
--- a/Engine/Source/Runtime/Renderer/Private/HeterogeneousVolumes/HeterogeneousVolumes.cpp
+++ b/Engine/Source/Runtime/Renderer/Private/HeterogeneousVolumes/HeterogeneousVolumes.cpp
@@ -15,7 +15,14 @@
 static TAutoConsoleVariable<int32> CVarHeterogeneousVolumes(
 	TEXT("r.HeterogeneousVolumes"),
 	1,
-	TEXT("Enables the Heterogeneous volume integrator (Default = 1)"),
+	TEXT("Whether to use Heterogeneous volumes by default (Default = 1)"),
+	ECVF_RenderThreadSafe
+);
+
+static TAutoConsoleVariable<int32> CVarHeterogeneousVolumesAllowed(
+	TEXT("r.HeterogeneousVolumes.Allowed"),
+	1,
+	TEXT("Whether Heterogeneous Volumes is allowed by scalability and device profiles (Default = 1)"),
 	ECVF_Scalability | ECVF_RenderThreadSafe
 );
 
@@ -348,7 +355,8 @@ DECLARE_GPU_STAT_NAMED(HeterogeneousVolumesStat, TEXT("HeterogeneousVolumes"));
 
 static bool IsHeterogeneousVolumesEnabled()
 {
-	return CVarHeterogeneousVolumes.GetValueOnRenderThread() != 0;
+	return CVarHeterogeneousVolumes.GetValueOnRenderThread() != 0 &&
+		CVarHeterogeneousVolumesAllowed.GetValueOnRenderThread() != 0;
 }
 
 bool ShouldHeterogeneousVolumesCastShadows()
