Lumen - increase number of World Space Radiance Cache probes updated per frame on Epic scalability level. This helps with lighting convergence on fast camera movement in large games.

#Synced-CL 51640903

diff --git a/Engine/Config/BaseScalability.ini b/Engine/Config/BaseScalability.ini
--- a/Engine/Config/BaseScalability.ini
+++ b/Engine/Config/BaseScalability.ini
@@ -376,7 +376,7 @@ r.LumenScene.Radiosity.ProbeSpacing=4
 r.LumenScene.Radiosity.HemisphereProbeResolution=4
 r.Lumen.TraceMeshSDFs.Allow=1
 r.Lumen.ScreenProbeGather.RadianceCache.ProbeResolution=32
-r.Lumen.ScreenProbeGather.RadianceCache.NumProbesToTraceBudget=100
+r.Lumen.ScreenProbeGather.RadianceCache.NumProbesToTraceBudget=150
 r.Lumen.ScreenProbeGather.DownsampleFactor=16
 r.Lumen.ScreenProbeGather.NumAdaptiveProbes=8
 r.Lumen.ScreenProbeGather.TracingOctahedronResolution=8
@@ -396,7 +396,7 @@ r.Lumen.TranslucencyVolume.GridPixelSize=32
 r.Lumen.TranslucencyVolume.TraceFromVolume=1
 r.Lumen.TranslucencyVolume.TracingOctahedronResolution=3
 r.Lumen.TranslucencyVolume.RadianceCache.ProbeResolution=8
-r.Lumen.TranslucencyVolume.RadianceCache.NumProbesToTraceBudget=70
+r.Lumen.TranslucencyVolume.RadianceCache.NumProbesToTraceBudget=100
 r.Lumen.HardwareRayTracing.HitLighting.Allowed=1
 r.RayTracing.Scene.BuildMode=1
 r.RayTracing.Geometry.SkeletalMeshes=1
