Fix casing of CVars for OcclusionFeedback settings in BaseDeviceProfiles.ini

#Synced-CL 52062538

diff --git a/Engine/Config/BaseDeviceProfiles.ini b/Engine/Config/BaseDeviceProfiles.ini
--- a/Engine/Config/BaseDeviceProfiles.ini
+++ b/Engine/Config/BaseDeviceProfiles.ini
@@ -1410,21 +1410,21 @@ BaseProfileName=Android_Mid
 DeviceType=Android
 BaseProfileName=Android_OpenXR
 +CVars=r.Mobile.Oculus.ForceSymmetric=1
-+Cvars=r.OcclusionFeedback.Enable=0
++CVars=r.OcclusionFeedback.Enable=0
 +CVars=r.Mobile.AdrenoOcclusionMode=1
 
 [Meta_Quest_Pro DeviceProfile]
 DeviceType=Android
 BaseProfileName=Android_OpenXR
 +CVars=r.Mobile.Oculus.ForceSymmetric=1
-+Cvars=r.OcclusionFeedback.Enable=0
++CVars=r.OcclusionFeedback.Enable=0
 +CVars=r.Mobile.AdrenoOcclusionMode=1
 
 [Oculus_Quest2 DeviceProfile]
 DeviceType=Android
 BaseProfileName=Android_OpenXR
 +CVars=r.Mobile.Oculus.ForceSymmetric=1
-+Cvars=r.OcclusionFeedback.Enable=0
++CVars=r.OcclusionFeedback.Enable=0
 +CVars=r.Mobile.AdrenoOcclusionMode=1
 +CVars=sg.ShadowQuality=1
 
