[BUGFIX] Interchange GLTF Pipelines for Force Reimport Material Instances

#rnx
#Synced-CL 52171822

diff --git a/Engine/Config/BaseEngine.ini b/Engine/Config/BaseEngine.ini
--- a/Engine/Config/BaseEngine.ini
+++ b/Engine/Config/BaseEngine.ini
@@ -3746,7 +3746,7 @@ UsePerContainerTocs=False
 
 [/Script/InterchangeEngine.InterchangeProjectSettings]
 ContentImportSettings=(DefaultPipelineStackOverride=((Materials, "Materials"),(Textures, "Textures"),(Sounds,"Audio")),ShowImportDialogOverride=((Textures,(bShowImportDialog=False)),(Sounds,(bShowImportDialog=False))),PipelineStacks=(("Assets", (Pipelines=("/Interchange/Pipelines/DefaultAssetsPipeline.DefaultAssetsPipeline"),PerTranslatorPipelines=((Translator="/Script/InterchangeImport.InterchangeMaterialXTranslator",Pipelines=("/Interchange/Pipelines/DefaultAssetsPipeline.DefaultAssetsPipeline","/Interchange/Pipelines/DefaultMaterialXPipeline.DefaultMaterialXPipeline")),(Translator="/Script/InterchangeImport.InterchangeGLTFTranslator",Pipelines=("/Interchange/Pipelines/DefaultGLTFAssetsPipeline.DefaultGLTFAssetsPipeline","/Interchange/Pipelines/DefaultGLTFPipeline.DefaultGLTFPipeline"))))),("Materials", (Pipelines=("/Interchange/Pipelines/DefaultMaterialPipeline.DefaultMaterialPipeline"),PerTranslatorPipelines=((Translator="/Script/InterchangeImport.InterchangeMaterialXTranslator",Pipelines=("/Interchange/Pipelines/DefaultMaterialPipeline.DefaultMaterialPipeline","/Interchange/Pipelines/DefaultMaterialXPipeline.DefaultMaterialXPipeline"))))),("Textures", (Pipelines=("/Interchange/Pipelines/DefaultTexturePipeline.DefaultTexturePipeline"))),("Audio", (Pipelines=("/Interchange/Pipelines/DefaultAudioPipeline.DefaultAudioPipeline")))),DefaultPipelineStack="Assets",ImportDialogClass="/Script/InterchangeEditorPipelines.InterchangePipelineConfigurationGeneric",bShowImportDialog=True)
-SceneImportSettings=(PipelineStacks=(("Scene", (Pipelines=("/Interchange/Pipelines/DefaultSceneAssetsPipeline.DefaultSceneAssetsPipeline","/Interchange/Pipelines/DefaultSceneLevelPipeline.DefaultSceneLevelPipeline"),PerTranslatorPipelines=((Translator="/Script/InterchangeImport.InterchangeMaterialXTranslator",Pipelines=("/Interchange/Pipelines/DefaultSceneAssetsPipeline.DefaultSceneAssetsPipeline","/Interchange/Pipelines/DefaultSceneLevelPipeline.DefaultSceneLevelPipeline","/Interchange/Pipelines/DefaultMaterialXPipeline.DefaultMaterialXPipeline")),(Translator="/Script/InterchangeImport.InterchangeGLTFTranslator",Pipelines=("/Interchange/Pipelines/DefaultGLTFSceneAssetsPipeline.DefaultGLTFSceneAssetsPipeline","/Interchange/Pipelines/DefaultSceneLevelPipeline.DefaultSceneLevelPipeline","/Interchange/Pipelines/DefaultGLTFPipeline.DefaultGLTFPipeline")))))),DefaultPipelineStack="Scene",ImportDialogClass="/Script/InterchangeEditorPipelines.InterchangePipelineConfigurationGeneric",bShowImportDialog=True)
+SceneImportSettings=(PipelineStacks=(("Scene", (Pipelines=("/Interchange/Pipelines/DefaultSceneAssetsPipeline.DefaultSceneAssetsPipeline","/Interchange/Pipelines/DefaultSceneLevelPipeline.DefaultSceneLevelPipeline"),PerTranslatorPipelines=((Translator="/Script/InterchangeImport.InterchangeMaterialXTranslator",Pipelines=("/Interchange/Pipelines/DefaultSceneAssetsPipeline.DefaultSceneAssetsPipeline","/Interchange/Pipelines/DefaultSceneLevelPipeline.DefaultSceneLevelPipeline","/Interchange/Pipelines/DefaultMaterialXPipeline.DefaultMaterialXPipeline")),(Translator="/Script/InterchangeImport.InterchangeGLTFTranslator",Pipelines=("/Interchange/Pipelines/DefaultGLTFSceneAssetsPipeline.DefaultGLTFSceneAssetsPipeline","/Interchange/Pipelines/DefaultGLTFPipeline.DefaultGLTFPipeline","/Interchange/Pipelines/DefaultSceneLevelPipeline.DefaultSceneLevelPipeline")))))),DefaultPipelineStack="Scene",ImportDialogClass="/Script/InterchangeEditorPipelines.InterchangePipelineConfigurationGeneric",bShowImportDialog=True)
 FilePickerClass=/Script/InterchangeEditorUtilities.InterchangeFilePickerGeneric
 bStaticMeshUseSmoothEdgesIfSmoothingInformationIsMissing=True
 GenericPipelineClass=/Script/InterchangePipelines.InterchangeGenericAssetsPipeline
diff --git a/Engine/Plugins/Interchange/Runtime/Source/Pipelines/Private/InterchangeglTFPipeline.cpp b/Engine/Plugins/Interchange/Runtime/Source/Pipelines/Private/InterchangeglTFPipeline.cpp
--- a/Engine/Plugins/Interchange/Runtime/Source/Pipelines/Private/InterchangeglTFPipeline.cpp
+++ b/Engine/Plugins/Interchange/Runtime/Source/Pipelines/Private/InterchangeglTFPipeline.cpp
@@ -335,18 +335,28 @@ UInterchangeMaterialInstanceFactoryNode* UGLTFPipelineSettings::BuildMaterialIns
 
 	FString MaterialFactoryNodeName = ShaderGraphNode->GetDisplayLabel();
 
-	// Cache reimport strategy flag from material factory node before replacing it
+	// Transfer properties from the Material Factory Node to Material Instance Factory Node
 	EReimportStrategyFlags ReimportStrategy = EReimportStrategyFlags::ApplyNoProperties;
+	bool bShouldForceReimport = false;
+
 	if (UInterchangeBaseMaterialFactoryNode* MaterialFactoryNode = Cast<UInterchangeBaseMaterialFactoryNode>(NodeContainer->GetFactoryNode(OldFactoryNodeUId)))
 	{
+		// Retrieve Attributes to Transfer
 		ReimportStrategy = MaterialFactoryNode->GetReimportStrategyFlags();
+		bShouldForceReimport = MaterialFactoryNode->ShouldForceNodeReimport();
 	}
 
 	if (UInterchangeMaterialInstanceFactoryNode* MaterialInstanceFactoryNode = NewObject<UInterchangeMaterialInstanceFactoryNode>(NodeContainer))
 	{
 		NodeContainer->SetupAndReplaceFactoryNode(MaterialInstanceFactoryNode, MaterialFactoryNodeUid, MaterialFactoryNodeName, EInterchangeNodeContainerType::FactoryData, OldFactoryNodeUId);
 
+		// Transfer the Attributes
+		if (bShouldForceReimport)
+		{
+			MaterialInstanceFactoryNode->SetForceNodeReimport();
+		}
 		MaterialInstanceFactoryNode->SetReimportStrategyFlags(ReimportStrategy);
+
 		MaterialInstanceFactoryNode->SetCustomParent(Parent);
 
 		UInterchangeEditorUtilitiesBase* EditorUtilities = UInterchangeManager::GetInterchangeManager().GetEditorUtilities();
