Remove r.Mobile.AntiAliasing=3 (MSAA) from the Android_OpenXR device profile. It will no longer override the AA mode set in project settings.
MSAA is already the default anti-aliasing setting for the VR Template, and devs should be able to use any AA mode they want on Quest and other Android-based XR devices.

#Synced-CL 52263667

diff --git a/Engine/Config/BaseDeviceProfiles.ini b/Engine/Config/BaseDeviceProfiles.ini
--- a/Engine/Config/BaseDeviceProfiles.ini
+++ b/Engine/Config/BaseDeviceProfiles.ini
@@ -1401,7 +1401,6 @@ DeviceType=Android
 BaseProfileName=Android_Mid
 +CVars=r.Android.DisableVulkanSupport=0
 +CVars=vr.MobileMultiView=1
-+CVars=r.Mobile.AntiAliasing=3
 +CVars=r.ForceStartInVR=1
 +CVars=xr.VRS.DynamicFoveation=1
 +CVars=r.Vulkan.VRSFormat=3
