Android: Fallback 3D textures to ETC compression
Replaces ASTC with ETC for 3D textures to address limitations on devices like the Pixel 10. While these devices may not explicitly report support for ETC 3D sampling, internal testing confirms it functions correctly, whereas ASTC fails
Allow Vulkan RHI on Pixel 10 devices
(submitting on behalf of Dmitriy.Dyomin)
#tests Local client testing
#Synced-CL 52194623

diff --git a/Engine/Config/BaseDeviceProfiles.ini b/Engine/Config/BaseDeviceProfiles.ini
--- a/Engine/Config/BaseDeviceProfiles.ini
+++ b/Engine/Config/BaseDeviceProfiles.ini
@@ -941,9 +941,9 @@ BaseProfileName=VisionPro
 +MatchProfile=(Profile="Android_PowerVR_GE8xxx",Match=((SourceType=SRC_GpuFamily,CompareType=CMP_Regex,MatchString="PowerVR Rogue GE8[0-9]+")))
 +MatchProfile=(Profile="Android_PowerVR_GM9xxx_Vulkan",Match=((SourceType=SRC_GpuFamily,CompareType=CMP_Regex,MatchString="PowerVR Rogue GM9[0-9]+"),(SourceType=SRC_VulkanAvailable,CompareType=CMP_Equal,MatchString="true")))
 +MatchProfile=(Profile="Android_PowerVR_GM9xxx",Match=((SourceType=SRC_GpuFamily,CompareType=CMP_Regex,MatchString="PowerVR Rogue GM9[0-9]+")))
-;+MatchProfile=(Profile="Android_PowerVR_Series_High_Vulkan",Match=((SourceType=SRC_GpuFamily,CompareType=CMP_Regex,MatchString="PowerVR [A-D]\\-Series [A-D]X[A-Z]+\\-([0-9]+)\\-[0-9]+"),(SourceType=SRC_PreviousRegexMatch,CompareType=CMP_GreaterEqual,MatchString="32"),(SourceType=SRC_VulkanAvailable,CompareType=CMP_Equal,MatchString="true")))
-;+MatchProfile=(Profile="Android_PowerVR_Series_Mid_Vulkan",Match=((SourceType=SRC_GpuFamily,CompareType=CMP_Regex,MatchString="PowerVR [A-D]\\-Series [A-D]X[A-Z]+\\-([0-9]+)\\-[0-9]+"),(SourceType=SRC_PreviousRegexMatch,CompareType=CMP_GreaterEqual,MatchString="16"),(SourceType=SRC_VulkanAvailable,CompareType=CMP_Equal,MatchString="true")))
-;+MatchProfile=(Profile="Android_PowerVR_Series_Low_Vulkan",Match=((SourceType=SRC_GpuFamily,CompareType=CMP_Regex,MatchString="PowerVR [A-D]\\-Series [A-D]X[A-Z]+\\-[0-9]+\\-[0-9]+"),(SourceType=SRC_VulkanAvailable,CompareType=CMP_Equal,MatchString="true")))
++MatchProfile=(Profile="Android_PowerVR_Series_High_Vulkan",Match=((SourceType=SRC_GpuFamily,CompareType=CMP_Regex,MatchString="PowerVR [A-D]\\-Series [A-D]X[A-Z]+\\-([0-9]+)\\-[0-9]+"),(SourceType=SRC_PreviousRegexMatch,CompareType=CMP_GreaterEqual,MatchString="32"),(SourceType=SRC_VulkanAvailable,CompareType=CMP_Equal,MatchString="true")))
++MatchProfile=(Profile="Android_PowerVR_Series_Mid_Vulkan",Match=((SourceType=SRC_GpuFamily,CompareType=CMP_Regex,MatchString="PowerVR [A-D]\\-Series [A-D]X[A-Z]+\\-([0-9]+)\\-[0-9]+"),(SourceType=SRC_PreviousRegexMatch,CompareType=CMP_GreaterEqual,MatchString="16"),(SourceType=SRC_VulkanAvailable,CompareType=CMP_Equal,MatchString="true")))
++MatchProfile=(Profile="Android_PowerVR_Series_Low_Vulkan",Match=((SourceType=SRC_GpuFamily,CompareType=CMP_Regex,MatchString="PowerVR [A-D]\\-Series [A-D]X[A-Z]+\\-[0-9]+\\-[0-9]+"),(SourceType=SRC_VulkanAvailable,CompareType=CMP_Equal,MatchString="true")))
 +MatchProfile=(Profile="Android_PowerVR_Series_High",Match=((SourceType=SRC_GpuFamily,CompareType=CMP_Regex,MatchString="PowerVR [A-D]\\-Series [A-D]X[A-Z]+\\-([0-9]+)\\-[0-9]+"),(SourceType=SRC_PreviousRegexMatch,CompareType=CMP_GreaterEqual,MatchString="32")))
 +MatchProfile=(Profile="Android_PowerVR_Series_Mid",Match=((SourceType=SRC_GpuFamily,CompareType=CMP_Regex,MatchString="PowerVR [A-D]\\-Series [A-D]X[A-Z]+\\-([0-9]+)\\-[0-9]+"),(SourceType=SRC_PreviousRegexMatch,CompareType=CMP_GreaterEqual,MatchString="16")))
 +MatchProfile=(Profile="Android_PowerVR_Series_Low",Match=((SourceType=SRC_GpuFamily,CompareType=CMP_Regex,MatchString="PowerVR [A-D]\\-Series [A-D]X[A-Z]+\\-[0-9]+\\-[0-9]+")))
diff --git a/Engine/Source/Developer/Android/AndroidTargetPlatformControls/Private/AndroidTargetPlatformControls.cpp b/Engine/Source/Developer/Android/AndroidTargetPlatformControls/Private/AndroidTargetPlatformControls.cpp
--- a/Engine/Source/Developer/Android/AndroidTargetPlatformControls/Private/AndroidTargetPlatformControls.cpp
+++ b/Engine/Source/Developer/Android/AndroidTargetPlatformControls/Private/AndroidTargetPlatformControls.cpp
@@ -564,6 +564,27 @@ void FAndroid_ASTCTargetPlatformControls::GetTextureFormats(const UTexture* Text
 				break;
 			}
 		}
+
+		if (Texture && Texture->GetTextureClass() == ETextureClass::Volume)
+		{
+			// Fallback to ETC for 3D textures, as ASTC for volume textures is unsupported on certain devices like the Pixel 10
+			const static FName VolumeTexRemap[][2] =
+			{
+				{ AndroidTexFormat::NameAutoASTC,	AndroidTexFormat::NameAutoETC2	},
+				{ FName(TEXT("ASTC_RGBA")),			AndroidTexFormat::NameETC2_RGBA	},
+				{ FName(TEXT("ASTC_RGB")),			AndroidTexFormat::NameETC2_RGB	},
+				{ FName(TEXT("ASTC_RGBA_HQ")),		AndroidTexFormat::NameETC2_RGBA	}
+			};
+
+			for (int32 RemapIndex = 0; RemapIndex < UE_ARRAY_COUNT(VolumeTexRemap); ++RemapIndex)
+			{
+				if (TextureFormatName == VolumeTexRemap[RemapIndex][0])
+				{
+					TextureFormatName = VolumeTexRemap[RemapIndex][1];
+					break;
+				}
+			}
+		}
 	}
 
 	bool bSupportASTCHDR = AndroidTargetPlatformSettings->UsesASTCHDR();
