[PCG] Add pcg.Quality CVar to Foliage scalability group

#Synced-CL 0

diff --git a/Engine/Config/BaseScalability.ini b/Engine/Config/BaseScalability.ini
--- a/Engine/Config/BaseScalability.ini
+++ b/Engine/Config/BaseScalability.ini
@@ -897,6 +897,7 @@ r.Translucency.HeterogeneousVolumes.Type=0
 [FoliageQuality@0]
 foliage.DensityScale=0
 grass.DensityScale=0
+pcg.Quality=0
 ;foliage.DiscardDataOnLoad=1 ; having this enabled will reduce ram use but will require reloading the level when changing FoliageQuality setting
 ;grass.DiscardDataOnLoad=1   ; having this enabled will reduce ram use but will require reloading the level when changing FoliageQuality setting
 ;foliage.LODDistanceScale=0
@@ -904,6 +905,7 @@ grass.DensityScale=0
 [FoliageQuality@1]
 foliage.DensityScale=0.4
 grass.DensityScale=0.4
+pcg.Quality=1
 ;foliage.DiscardDataOnLoad=0
 ;grass.DiscardDataOnLoad=0
 ;foliage.LODDistanceScale=0.4
@@ -911,6 +913,7 @@ grass.DensityScale=0.4
 [FoliageQuality@2]
 foliage.DensityScale=0.8
 grass.DensityScale=0.8
+pcg.Quality=2
 ;foliage.DiscardDataOnLoad=0
 ;grass.DiscardDataOnLoad=0
 ;foliage.LODDistanceScale=0.8
@@ -918,6 +921,7 @@ grass.DensityScale=0.8
 [FoliageQuality@3]
 foliage.DensityScale=1.0
 grass.DensityScale=1.0
+pcg.Quality=3
 ;foliage.DiscardDataOnLoad=0
 ;grass.DiscardDataOnLoad=0
 ;foliage.LODDistanceScale=1.0
@@ -925,6 +929,7 @@ grass.DensityScale=1.0
 [FoliageQuality@Cine]
 foliage.DensityScale=1.0
 grass.DensityScale=1.0
+pcg.Quality=4
 ;foliage.DiscardDataOnLoad=0
 ;grass.DiscardDataOnLoad=0
 ;foliage.LODDistanceScale=1.0
