Disable chunked PSO cache on Quest devices

The chunked PSO cache has the following tradeoffs:

Pro: Drivers which fully serialize access to their PSO caches get better concurrency
Pro: Drivers which don't scale well to large PSO caches get faster lookups
Con: Drivers which can reuse cached information for pipelines which are different but similar may have cache misses due to the pipelines being in different chunks
Con: Drivers which can reuse cached information for pipelines which are different but similar may use extra memory storing multiple copies of this information in different chunks

On Quest devices, the driver doesn't benefit much from the Pros, and is hindered by the Cons. It's better to turn this feature off

#Synced-CL 52421815

diff --git a/Engine/Config/BaseDeviceProfiles.ini b/Engine/Config/BaseDeviceProfiles.ini
--- a/Engine/Config/BaseDeviceProfiles.ini
+++ b/Engine/Config/BaseDeviceProfiles.ini
@@ -1414,6 +1414,7 @@ BaseProfileName=Android_OpenXR
 +CVars=r.Mobile.Oculus.ForceSymmetric=1
 +CVars=r.OcclusionFeedback.Enable=0
 +CVars=r.Mobile.AdrenoOcclusionMode=1
++CVars=r.Vulkan.UseChunkedPSOCache=0
 
 [Meta_Quest_Pro DeviceProfile]
 DeviceType=Android
@@ -1421,6 +1422,7 @@ BaseProfileName=Android_OpenXR
 +CVars=r.Mobile.Oculus.ForceSymmetric=1
 +CVars=r.OcclusionFeedback.Enable=0
 +CVars=r.Mobile.AdrenoOcclusionMode=1
++CVars=r.Vulkan.UseChunkedPSOCache=0
 
 [Oculus_Quest2 DeviceProfile]
 DeviceType=Android
@@ -1429,6 +1431,7 @@ BaseProfileName=Android_OpenXR
 +CVars=r.OcclusionFeedback.Enable=0
 +CVars=r.Mobile.AdrenoOcclusionMode=1
 +CVars=sg.ShadowQuality=1
++CVars=r.Vulkan.UseChunkedPSOCache=0
 
 [Oculus_Quest DeviceProfile]
 DeviceType=Android
@@ -1437,6 +1440,7 @@ BaseProfileName=Android_OpenXR
 +CVars=fx.NiagaraAllowGPUParticles=0
 +CVars=FX.AllowGPUSorting=0
 +CVars=sg.ShadowQuality=1
++CVars=r.Vulkan.UseChunkedPSOCache=0
 
 [Mac DeviceProfile]
 DeviceType=Mac
