Medium GI quality disables r.MeshCardRepresentation.SkeletalMesh, where skeletal meshes can't be intersected by HWRT anyway
* r.MeshCardRepresentation.SkeletalMesh now only recreates skeletal mesh components, and can be changed by scalability

#Synced-CL 52380299

diff --git a/Engine/Config/BaseScalability.ini b/Engine/Config/BaseScalability.ini
--- a/Engine/Config/BaseScalability.ini
+++ b/Engine/Config/BaseScalability.ini
@@ -311,6 +311,7 @@ r.Lumen.TranslucencyVolume.Enable=0
 r.Lumen.HardwareRayTracing.HitLighting.Allowed=0
 r.RayTracing.Scene.BuildMode=0
 r.RayTracing.Geometry.SkeletalMeshes=0
+r.MeshCardRepresentation.SkeletalMesh=0
 r.RayTracing.DynamicGeometry.MaxUpdatePrimitivesPerFrame=1000
 
 [GlobalIlluminationQuality@2]
@@ -355,6 +356,7 @@ r.Lumen.TranslucencyVolume.RadianceCache.NumProbesToTraceBudget=70
 r.Lumen.HardwareRayTracing.HitLighting.Allowed=0
 r.RayTracing.Scene.BuildMode=0
 r.RayTracing.Geometry.SkeletalMeshes=1
+r.MeshCardRepresentation.SkeletalMesh=1
 r.RayTracing.DynamicGeometry.MaxUpdatePrimitivesPerFrame=-1
 
 [GlobalIlluminationQuality@3]
@@ -400,6 +402,7 @@ r.Lumen.TranslucencyVolume.RadianceCache.NumProbesToTraceBudget=100
 r.Lumen.HardwareRayTracing.HitLighting.Allowed=1
 r.RayTracing.Scene.BuildMode=1
 r.RayTracing.Geometry.SkeletalMeshes=1
+r.MeshCardRepresentation.SkeletalMesh=1
 r.RayTracing.DynamicGeometry.MaxUpdatePrimitivesPerFrame=-1
 
 [GlobalIlluminationQuality@Cine]
@@ -445,6 +448,7 @@ r.Lumen.TranslucencyVolume.RadianceCache.NumProbesToTraceBudget=350
 r.Lumen.HardwareRayTracing.HitLighting.Allowed=1
 r.RayTracing.Scene.BuildMode=1
 r.RayTracing.Geometry.SkeletalMeshes=1
+r.MeshCardRepresentation.SkeletalMesh=1
 r.RayTracing.DynamicGeometry.MaxUpdatePrimitivesPerFrame=-1
 
 ;-----------------------------------------------------------------------------------------------------------------
diff --git a/Engine/Source/Runtime/Engine/Private/SkeletalMeshSceneProxy.cpp b/Engine/Source/Runtime/Engine/Private/SkeletalMeshSceneProxy.cpp
--- a/Engine/Source/Runtime/Engine/Private/SkeletalMeshSceneProxy.cpp
+++ b/Engine/Source/Runtime/Engine/Private/SkeletalMeshSceneProxy.cpp
@@ -71,8 +71,8 @@ static TAutoConsoleVariable<int32> CVarMeshCardRepresentationSkeletalMesh(
 	TEXT("r.MeshCardRepresentation.SkeletalMesh"),
 	1,
 	TEXT("Whether to allow generating mesh cards for skeletal meshes."),
-	FConsoleVariableDelegate::CreateLambda([](IConsoleVariable* InVariable) { FGlobalComponentRecreateRenderStateContext Context; }),
-	ECVF_RenderThreadSafe);
+	FConsoleVariableDelegate::CreateLambda([](IConsoleVariable* InVariable) { FSkinnedMeshComponentRecreateRenderStateContext Context(nullptr); }),
+	ECVF_Scalability | ECVF_RenderThreadSafe);
 
 const FQuat SphylBasis(FVector(1.0f / FMath::Sqrt(2.0f), 0.0f, 1.0f / FMath::Sqrt(2.0f)), UE_PI);
 
diff --git a/Engine/Source/Runtime/Engine/Private/SkinnedAsset.cpp b/Engine/Source/Runtime/Engine/Private/SkinnedAsset.cpp
--- a/Engine/Source/Runtime/Engine/Private/SkinnedAsset.cpp
+++ b/Engine/Source/Runtime/Engine/Private/SkinnedAsset.cpp
@@ -591,7 +591,7 @@ FSkinnedMeshComponentRecreateRenderStateContext::FSkinnedMeshComponentRecreateRe
 	TSet<FSceneInterface*> Scenes;
 
 	FObjectCacheContextScope ObjectCacheScope;
-	for (ISkinnedMeshComponent* Component : ObjectCacheScope.GetContext().GetSkinnedMeshComponents(InSkinnedAsset))
+	for (ISkinnedMeshComponent* Component : (InSkinnedAsset ? ObjectCacheScope.GetContext().GetSkinnedMeshComponents(InSkinnedAsset) : ObjectCacheScope.GetContext().GetSkinnedMeshComponents()))
 	{
 		IPrimitiveComponent* PrimitiveComponent = Component->GetPrimitiveComponentInterface();
 
