Cleaned up some material expressions groups:
- Moved "Utils" expressions under "Decals"
- Merged "VectroOps" into "Vectors"
- Merged "DistanceFieldsRendering" into "DistanceFields"

#rnx
#Synced-CL 52505914

diff --git a/Engine/Config/BaseMaterialExpressions.ini b/Engine/Config/BaseMaterialExpressions.ini
--- a/Engine/Config/BaseMaterialExpressions.ini
+++ b/Engine/Config/BaseMaterialExpressions.ini
@@ -65,7 +65,7 @@ bShaderInputData=true
 
 [/Script/Engine.MaterialExpressionAppendVector]
 +MenuCategories=NSLOCTEXT("MaterialExpression", "Math", "Math")
-+MenuCategories=NSLOCTEXT("MaterialExpression", "VectorOps", "VectorOps")
++MenuCategories=NSLOCTEXT("MaterialExpression", "Vectors", "Vectors")
 
 [/Script/Engine.MaterialExpressionArccosine]
 +MenuCategories=NSLOCTEXT("MaterialExpression", "Math", "Math")
@@ -183,7 +183,7 @@ bShaderInputData=true
 
 [/Script/Engine.MaterialExpressionComponentMask]
 +MenuCategories=NSLOCTEXT("MaterialExpression", "Math", "Math")
-+MenuCategories=NSLOCTEXT("MaterialExpression", "VectorOps", "VectorOps")
++MenuCategories=NSLOCTEXT("MaterialExpression", "Vectors", "Vectors")
 
 [/Script/Engine.MaterialExpressionComposite]
 +MenuCategories=NSLOCTEXT("MaterialExpression", "Utility", "Utility")
@@ -211,7 +211,7 @@ bShaderInputData=true
 
 [/Script/Engine.MaterialExpressionCrossProduct]
 +MenuCategories=NSLOCTEXT("MaterialExpression", "Math", "Math")
-+MenuCategories=NSLOCTEXT("MaterialExpression", "VectorOps", "VectorOps")
++MenuCategories=NSLOCTEXT("MaterialExpression", "Vectors", "Vectors")
 
 [/Script/Engine.MaterialExpressionScalarParameter]
 +MenuCategories=NSLOCTEXT("MaterialExpression", "Parameters", "Parameters")
@@ -236,21 +236,21 @@ bShaderInputData=true
 +MenuCategories=NSLOCTEXT("MaterialExpression", "Utility", "Utility")
 
 [/Script/Engine.MaterialExpressionDecalColor]
-+MenuCategories=NSLOCTEXT("MaterialExpression", "Utils", "Utils")
++MenuCategories=NSLOCTEXT("MaterialExpression", "Decals", "Decals")
 bShaderInputData=true
 +ExternalCodeIdentifiers="DecalColor"
 
 [/Script/Engine.MaterialExpressionDecalDerivative]
-+MenuCategories=NSLOCTEXT("MaterialExpression", "Utils", "Utils")
++MenuCategories=NSLOCTEXT("MaterialExpression", "Decals", "Decals")
 bShaderInputData=true
 
 [/Script/Engine.MaterialExpressionDecalLifetimeOpacity]
-+MenuCategories=NSLOCTEXT("MaterialExpression", "Utils", "Utils")
++MenuCategories=NSLOCTEXT("MaterialExpression", "Decals", "Decals")
 bShaderInputData=true
 +ExternalCodeIdentifiers="DecalLifetimeOpacity"
 
 [/Script/Engine.MaterialExpressionDecalMipmapLevel]
-+MenuCategories=NSLOCTEXT("MaterialExpression", "Utils", "Utils")
++MenuCategories=NSLOCTEXT("MaterialExpression", "Decals", "Decals")
 +ExternalCodeIdentifiers="DecalMipmapLevel"
 
 [/Script/Engine.MaterialExpressionDeltaTime]
@@ -267,7 +267,7 @@ bShaderInputData=true
 +ExternalCodeIdentifiers="DepthOfFieldFunction"
 
 [/Script/Engine.MaterialExpressionDeriveNormalZ]
-+MenuCategories=NSLOCTEXT("MaterialExpression", "VectorOps", "VectorOps")
++MenuCategories=NSLOCTEXT("MaterialExpression", "Vectors", "Vectors")
 
 [/Script/Engine.MaterialExpressionDesaturation]
 +MenuCategories=NSLOCTEXT("MaterialExpression", "Color", "Color")
@@ -287,7 +287,7 @@ bShaderInputData=true
 +MenuCategories=NSLOCTEXT("MaterialExpression", "Utility", "Utility")
 
 [/Script/Engine.MaterialExpressionDistanceFieldsRenderingSwitch]
-+MenuCategories=NSLOCTEXT("MaterialExpression", "DistanceFieldsRendering", "DistanceFieldsRendering")
++MenuCategories=NSLOCTEXT("MaterialExpression", "DistanceFields", "DistanceFields")
 
 [/Script/Engine.MaterialExpressionDistanceToNearestSurface]
 +MenuCategories=NSLOCTEXT("MaterialExpression", "Utility", "Utility")
@@ -297,7 +297,7 @@ bShaderInputData=true
 
 [/Script/Engine.MaterialExpressionDotProduct]
 +MenuCategories=NSLOCTEXT("MaterialExpression", "Math", "Math")
-+MenuCategories=NSLOCTEXT("MaterialExpression", "VectorOps", "VectorOps")
++MenuCategories=NSLOCTEXT("MaterialExpression", "Vectors", "Vectors")
 
 [/Script/Engine.MaterialExpressionDoubleVectorParameter]
 +MenuCategories=NSLOCTEXT("MaterialExpression", "Parameters", "Parameters")
@@ -499,7 +499,7 @@ bShaderInputData=true
 
 [/Script/Engine.MaterialExpressionNormalize]
 +MenuCategories=NSLOCTEXT("MaterialExpression", "Math", "Math")
-+MenuCategories=NSLOCTEXT("MaterialExpression", "VectorOps", "VectorOps")
++MenuCategories=NSLOCTEXT("MaterialExpression", "Vectors", "Vectors")
 
 [/Script/Engine.MaterialExpressionObjectBounds]
 +MenuCategories=NSLOCTEXT("MaterialExpression", "Vectors", "Vectors")
@@ -952,10 +952,10 @@ bShaderInputData=true
 bShaderInputData=true
 
 [/Script/Engine.MaterialExpressionTransform]
-+MenuCategories=NSLOCTEXT("MaterialExpression", "VectorOps", "VectorOps")
++MenuCategories=NSLOCTEXT("MaterialExpression", "Vectors", "Vectors")
 
 [/Script/Engine.MaterialExpressionTransformPosition]
-+MenuCategories=NSLOCTEXT("MaterialExpression", "VectorOps", "VectorOps")
++MenuCategories=NSLOCTEXT("MaterialExpression", "Vectors", "Vectors")
 
 [/Script/Engine.MaterialExpressionTruncate]
 +MenuCategories=NSLOCTEXT("MaterialExpression", "Math", "Math")
