Remove shader stage restrictions from GetPerInstanceRandom().
Fixes translation errors in new material translator.

#rnx

#Synced-CL 52572089

diff --git a/Engine/Config/BaseMaterialExpressions.ini b/Engine/Config/BaseMaterialExpressions.ini
--- a/Engine/Config/BaseMaterialExpressions.ini
+++ b/Engine/Config/BaseMaterialExpressions.ini
@@ -1087,7 +1087,7 @@ bShaderInputData=true
 +ExternalCodeDeclarations=(ReturnType=Float,			Name="GetPerInstanceCustomData",		Definition="GetPerInstanceCustomData(Parameters, $0, $1)", Derivative=Zero, EnvironmentDefines=((Name="USES_PER_INSTANCE_CUSTOM_DATA")))
 +ExternalCodeDeclarations=(ReturnType=Float3,			Name="GetPerInstanceCustomData3Vector",	Definition="GetPerInstanceCustomData3Vector(Parameters, $0, $1)", Derivative=Zero, EnvironmentDefines=((Name="USES_PER_INSTANCE_CUSTOM_DATA")))
 +ExternalCodeDeclarations=(ReturnType=Float,			Name="GetPerInstanceFadeAmount",		Definition="GetPerInstanceFadeAmount(Parameters)", ShaderFrequency=Vertex|Pixel, Derivative=Zero, bIsInlined=true, EnvironmentDefines=((Name="USES_PER_INSTANCE_FADE_AMOUNT")))
-+ExternalCodeDeclarations=(ReturnType=Float,			Name="GetPerInstanceRandom",			Definition="GetPerInstanceRandom(Parameters)", ShaderFrequency=Vertex|Pixel, Derivative=Zero, bIsInlined=true, EnvironmentDefines=((Name="NEEDS_PER_INSTANCE_RANDOM_PS", ShaderFrequency=Pixel)))
++ExternalCodeDeclarations=(ReturnType=Float,			Name="GetPerInstanceRandom",			Definition="GetPerInstanceRandom(Parameters)", Derivative=Zero, bIsInlined=true, EnvironmentDefines=((Name="NEEDS_PER_INSTANCE_RANDOM_PS")))
 
 +ExternalCodeDeclarations=(ReturnType=Float,	Name="GetCloudSampleAltitude",								Definition="MaterialExpressionCloudSampleAltitude(Parameters)")
 +ExternalCodeDeclarations=(ReturnType=Float,	Name="GetCloudSampleAltitudeInLayer",						Definition="MaterialExpressionCloudSampleAltitudeInLayer(Parameters)")
