[MegaLights] Default scalability settings

#rnx
#Synced-CL 52639728

diff --git a/Engine/Config/BaseScalability.ini b/Engine/Config/BaseScalability.ini
--- a/Engine/Config/BaseScalability.ini
+++ b/Engine/Config/BaseScalability.ini
@@ -153,6 +153,9 @@ r.Shadow.Virtual.SMRT.SamplesPerRayDirectional=2
 r.Shadow.Virtual.SMRT.RayCountLocal=0
 r.Shadow.Virtual.SMRT.SamplesPerRayLocal=2
 r.Shadow.Virtual.MarkCoarsePagesLocal=2
+r.MegaLights.NumSamplesPerPixel=2
+r.MegaLights.DownsampleMode=2
+r.MegaLights.ScreenTraces.Quality=0
 r.MegaLights.FrontLayerTranslucency.Allow=0
 
 [ShadowQuality@1]
@@ -179,6 +182,9 @@ r.Shadow.Virtual.SMRT.SamplesPerRayDirectional=2
 r.Shadow.Virtual.SMRT.RayCountLocal=4
 r.Shadow.Virtual.SMRT.SamplesPerRayLocal=2
 r.Shadow.Virtual.MarkCoarsePagesLocal=2
+r.MegaLights.NumSamplesPerPixel=2
+r.MegaLights.DownsampleMode=2
+r.MegaLights.ScreenTraces.Quality=0
 r.MegaLights.FrontLayerTranslucency.Allow=0
 
 [ShadowQuality@2]
@@ -208,6 +214,9 @@ r.Shadow.Virtual.SMRT.SamplesPerRayDirectional=4
 r.Shadow.Virtual.SMRT.RayCountLocal=4
 r.Shadow.Virtual.SMRT.SamplesPerRayLocal=4
 r.Shadow.Virtual.MarkCoarsePagesLocal=2
+r.MegaLights.NumSamplesPerPixel=4
+r.MegaLights.DownsampleMode=2
+r.MegaLights.ScreenTraces.Quality=1
 r.MegaLights.FrontLayerTranslucency.Allow=0
 
 [ShadowQuality@3]
@@ -237,6 +246,9 @@ r.Shadow.Virtual.SMRT.SamplesPerRayDirectional=4
 r.Shadow.Virtual.SMRT.RayCountLocal=8
 r.Shadow.Virtual.SMRT.SamplesPerRayLocal=4
 r.Shadow.Virtual.MarkCoarsePagesLocal=1
+r.MegaLights.NumSamplesPerPixel=4
+r.MegaLights.DownsampleMode=2
+r.MegaLights.ScreenTraces.Quality=1
 r.MegaLights.FrontLayerTranslucency.Allow=1
 
 [ShadowQuality@Cine]
@@ -266,6 +278,9 @@ r.Shadow.Virtual.SMRT.SamplesPerRayDirectional=8
 r.Shadow.Virtual.SMRT.RayCountLocal=16
 r.Shadow.Virtual.SMRT.SamplesPerRayLocal=8
 r.Shadow.Virtual.MarkCoarsePagesLocal=1
+r.MegaLights.NumSamplesPerPixel=4
+r.MegaLights.DownsampleMode=0
+r.MegaLights.ScreenTraces.Quality=1
 r.MegaLights.FrontLayerTranslucency.Allow=1
 
 ;-----------------------------------------------------------------------------------------------------------------
